[Insights] Massive refactoring of TraceInsights module:

diff --git a/Engine/Source/Developer/TraceInsightsCore/Public/InsightsCore/Table/Widgets/STableTreeView.h b/Engine/Source/Developer/TraceInsightsCore/Public/InsightsCore/Table/Widgets/STableTreeView.h
new file mode 100644
--- /dev/null
+++ b/Engine/Source/Developer/TraceInsightsCore/Public/InsightsCore/Table/Widgets/STableTreeView.h
@@ -0,0 +1,821 @@
+// Copyright Epic Games, Inc. All Rights Reserved.
+
+#pragma once
+
+#include "CoreMinimal.h"
+
+#include "Framework/MultiBox/MultiBoxExtender.h"
+#include "Misc/EnumClassFlags.h"
+#include "Misc/FilterCollection.h"
+#include "Misc/TextFilter.h"
+#include "SlateFwd.h"
+#include "Widgets/DeclarativeSyntaxSupport.h"
+#include "Widgets/Input/SComboBox.h"
+#include "Widgets/Navigation/SBreadcrumbTrail.h"
+#include "Widgets/SCompoundWidget.h"
+#include "Widgets/SWidget.h"
+#include "Widgets/Views/SHeaderRow.h"
+#include "Widgets/Views/STableViewBase.h"
+#include "Widgets/Views/STreeView.h"
+
+#include "InsightsCore/Common/AsyncOperationProgress.h"
+#include "InsightsCore/Common/IAsyncOperationStatusProvider.h"
+#include "InsightsCore/Common/Stopwatch.h"
+#include "InsightsCore/Filter/ViewModels/Filters.h"
+#include "InsightsCore/Table/ViewModels/TableTreeNode.h"
+
+#include <atomic>
+
+class FMenuBuilder;
+class FUICommandList;
+
+namespace UE::Insights
+{
+
+class FFilterConfigurator;
+class FTable;
+class FTableColumn;
+class FTreeNodeGrouping;
+class ITableCellValueSorter;
+
+////////////////////////////////////////////////////////////////////////////////////////////////////
+
+/** The filter collection - used for updating the list of tree nodes. */
+typedef TFilterCollection<const FTableTreeNodePtr&> FTableTreeNodeFilterCollection;
+
+/** The text based filter - used for updating the list of tree nodes. */
+typedef TTextFilter<const FTableTreeNodePtr&> FTableTreeNodeTextFilter;
+
+////////////////////////////////////////////////////////////////////////////////////////////////////
+
+enum class TRACEINSIGHTSCORE_API EAsyncOperationType : uint32
+{
+	NodeFiltering       = 1 << 0,
+	Grouping            = 1 << 1,
+	Aggregation         = 1 << 2,
+	Sorting             = 1 << 3,
+	HierarchyFiltering  = 1 << 4,
+};
+
+ENUM_CLASS_FLAGS(EAsyncOperationType);
+
+struct TRACEINSIGHTSCORE_API FTableColumnConfig
+{
+	FName ColumnId;
+	bool bIsVisible;
+	float Width;
+};
+
+class TRACEINSIGHTSCORE_API ITableTreeViewPreset
+{
+public:
+	virtual FText GetName() const = 0;
+	virtual FText GetToolTip() const = 0;
+	virtual FName GetSortColumn() const = 0;
+	virtual EColumnSortMode::Type GetSortMode() const = 0;
+	virtual void SetCurrentGroupings(const TArray<TSharedPtr<FTreeNodeGrouping>>& InAvailableGroupings, TArray<TSharedPtr<FTreeNodeGrouping>>& InOutCurrentGroupings) const = 0;
+	virtual void GetColumnConfigSet(TArray<FTableColumnConfig>& InOutConfigSet) const = 0;
+};
+
+class TRACEINSIGHTSCORE_API FTableTaskCancellationToken
+{
+public:
+	FTableTaskCancellationToken()
+		: bCancel(false)
+	{}
+
+	bool ShouldCancel() { return bCancel.load(); }
+	void Cancel() { bCancel.store(true); }
+
+private:
+	std::atomic<bool> bCancel;
+};
+
+struct TRACEINSIGHTSCORE_API FTableTaskInfo
+{
+	FGraphEventRef Event;
+	TSharedPtr<FTableTaskCancellationToken> CancellationToken;
+};
+
+////////////////////////////////////////////////////////////////////////////////////////////////////
+/**
+ * A custom widget used to display the list of tree nodes.
+ */
+class TRACEINSIGHTSCORE_API STableTreeView : public SCompoundWidget, public IAsyncOperationStatusProvider
+{
+	friend class FTableTreeViewNodeFilteringAsyncTask;
+	friend class FTableTreeViewSortingAsyncTask;
+	friend class FTableTreeViewGroupingAsyncTask;
+	friend class FTableTreeViewHierarchyFilteringAsyncTask;
+	friend class FTableTreeViewAsyncCompleteTask;
+	friend class FSearchForItemToSelectTask;
+	friend class FSelectNodeByTableRowIndexTask;
+
+public:
+	/** Default constructor. */
+	STableTreeView();
+
+	/** Virtual destructor. */
+	virtual ~STableTreeView();
+
+	SLATE_BEGIN_ARGS(STableTreeView) {}
+	SLATE_END_ARGS()
+
+	/**
+	 * Construct this widget
+	 * @param InArgs - The declaration data for this widget
+	 */
+	void Construct(const FArguments& InArgs, TSharedPtr<FTable> InTablePtr);
+
+	TSharedPtr<STreeView<FTableTreeNodePtr>> GetInnerTreeView() const { return TreeView; }
+
+	TSharedPtr<FTable>& GetTable() { return Table; }
+	const TSharedPtr<FTable>& GetTable() const { return Table; }
+
+	virtual void Reset();
+
+	void RebuildColumns();
+
+	/**
+	 * Ticks this widget.  Override in derived classes, but always call the parent implementation.
+	 *
+	 * @param  AllottedGeometry The space allotted for this widget
+	 * @param  InCurrentTime  Current absolute real time
+	 * @param  InDeltaTime  Real time passed since last tick
+	 */
+	virtual void Tick(const FGeometry& AllottedGeometry, const double InCurrentTime, const float InDeltaTime) override;
+
+	/**
+	 * Rebuilds the tree (if necessary).
+	 * @param bResync - If true, it forces a resync even if the list did not changed since last sync.
+	 */
+	virtual void RebuildTree(bool bResync);
+
+	FTableTreeNodePtr GetNodeByTableRowIndex(int32 RowIndex) const;
+	void SelectNodeByTableRowIndex(int32 RowIndex);
+	bool IsRunningAsyncUpdate() { return bIsUpdateRunning;  }
+
+	void OnClose();
+
+	virtual FReply OnKeyDown(const FGeometry& MyGeometry, const FKeyEvent& InKeyEvent) override;
+
+	////////////////////////////////////////////////////////////////////////////////////////////////////
+	// IAsyncOperationStatusProvider implementation
+
+	virtual bool IsRunning() const override { return bIsUpdateRunning; }
+
+	virtual double GetAllOperationsDuration() override;
+	virtual double GetCurrentOperationDuration() override { return 0.0; }
+	virtual uint32 GetOperationCount() const override { return 1; }
+	virtual FText GetCurrentOperationName() const override;
+
+	////////////////////////////////////////////////////////////////////////////////////////////////////
+
+	/** Sets a log listing name to be used for any errors or warnings. Must be preregistered by the caller with the MessageLog module. */
+	void SetLogListingName(const FName& InLogListingName) { LogListingName = InLogListingName; }
+	const FName& GetLogListingName() { return LogListingName; }
+
+	/** Gets the table row nodes. Each node corresponds to a table row. Index in this array corresponds to RowIndex in source table. */
+	const TArray<FTableTreeNodePtr>& GetTableRowNodes() const { return TableRowNodes; }
+
+	/** Gets the available groupings. */
+	const TArray<TSharedPtr<FTreeNodeGrouping>>& GetAvailableGroupings() const { return AvailableGroupings; }
+
+	/** Sets the current groupings. */
+	void SetCurrentGroupings(TArray<TSharedPtr<FTreeNodeGrouping>>& InCurrentGroupings);
+
+protected:
+	void InitCommandList();
+
+	virtual void ConstructWidget(TSharedPtr<FTable> InTablePtr);
+	virtual TSharedRef<SWidget> ConstructHierarchyBreadcrumbTrail();
+	virtual TSharedPtr<SWidget> ConstructToolbar() { return nullptr; }
+	virtual TSharedPtr<SWidget> ConstructFooter() { return nullptr; }
+	virtual void ConstructHeaderArea(TSharedRef<SVerticalBox> InWidgetContent);
+	virtual void ConstructFooterArea(TSharedRef<SVerticalBox> InWidgetContent);
+
+	void UpdateTree();
+
+	////////////////////////////////////////////////////////////////////////////////////////////////////
+	// Tree View - Context Menu
+
+	TSharedPtr<SWidget> TreeView_GetMenuContent();
+	void TreeView_BuildSortByMenu(FMenuBuilder& MenuBuilder);
+	void TreeView_BuildViewColumnMenu(FMenuBuilder& MenuBuilder);
+	void TreeView_BuildExportMenu(FMenuBuilder& MenuBuilder);
+
+	bool ContextMenu_CopySelectedToClipboard_CanExecute() const;
+	void ContextMenu_CopySelectedToClipboard_Execute();
+	bool ContextMenu_CopyColumnToClipboard_CanExecute() const;
+	void ContextMenu_CopyColumnToClipboard_Execute();
+	bool ContextMenu_CopyColumnTooltipToClipboard_CanExecute() const;
+	void ContextMenu_CopyColumnTooltipToClipboard_Execute();
+	bool ContextMenu_ExpandSubtree_CanExecute() const;
+	void ContextMenu_ExpandSubtree_Execute();
+	bool ContextMenu_ExpandCriticalPath_CanExecute() const;
+	void ContextMenu_ExpandCriticalPath_Execute();
+	bool ContextMenu_CollapseSubtree_CanExecute() const;
+	void ContextMenu_CollapseSubtree_Execute();
+	bool ContextMenu_ExportToFile_CanExecute() const;
+	void ContextMenu_ExportToFile_Execute(bool bInExportCollapsed, bool InExportLeafs);
+
+	////////////////////////////////////////////////////////////////////////////////////////////////////
+	// Tree View - Columns' Header
+
+	void InitializeAndShowHeaderColumns();
+
+	FText GetColumnHeaderText(const FName ColumnId) const;
+
+	TSharedRef<SWidget> TreeViewHeaderRow_GenerateColumnMenu(const FTableColumn& Column);
+
+	////////////////////////////////////////////////////////////////////////////////////////////////////
+	// Tree View - Misc
+
+	void TreeView_Refresh();
+
+	/**
+	 * Called by STreeView to retrieves the children for the specified parent item.
+	 * @param InParent    - The parent node to retrieve the children from.
+	 * @param OutChildren - List of children for the parent node.
+	 */
+	void TreeView_OnGetChildren(FTableTreeNodePtr InParent, TArray<FTableTreeNodePtr>& OutChildren);
+
+	/** Called by STreeView when selection has changed. */
+	virtual void TreeView_OnSelectionChanged(FTableTreeNodePtr SelectedItem, ESelectInfo::Type SelectInfo);
+
+	/** Called by STreeView when a tree node is expanded or collapsed. */
+	virtual void TreeView_OnExpansionChanged(FTableTreeNodePtr TreeNode, bool bShouldBeExpanded);
+
+	/** Called by STreeView when a tree item is double clicked. */
+	virtual void TreeView_OnMouseButtonDoubleClick(FTableTreeNodePtr TreeNode);
+
+	////////////////////////////////////////////////////////////////////////////////////////////////////
+	// Tree View - Table Row
+
+	/** Called by STreeView to generate a table row for the specified item. */
+	TSharedRef<ITableRow> TreeView_OnGenerateRow(FTableTreeNodePtr TreeNode, const TSharedRef<STableViewBase>& OwnerTable);
+
+	bool TableRow_ShouldBeEnabled(FTableTreeNodePtr NodePtr) const;
+
+	void TableRow_SetHoveredCell(TSharedPtr<FTable> TablePtr, TSharedPtr<FTableColumn> ColumnPtr, FTableTreeNodePtr NodePtr);
+	EHorizontalAlignment TableRow_GetColumnOutlineHAlignment(const FName ColumnId) const;
+
+	FText TableRow_GetHighlightText() const;
+	FName TableRow_GetHighlightedNodeName() const;
+
+	////////////////////////////////////////////////////////////////////////////////////////////////////
+	// Node Filtering (TableRowNodes --> FilteredNodes)
+
+	void InitNodeFiltering();
+
+	void OnNodeFilteringChanged();
+	bool ScheduleNodeFilteringAsyncOperationIfNeeded();
+	void ScheduleNodeFilteringAsyncOperation();
+	FGraphEventRef StartNodeFilteringTask(FGraphEventRef Prerequisite = nullptr);
+
+	void ApplyNodeFiltering();
+
+	virtual bool FilterNode(const FFilterConfigurator& InFilterConfigurator, const FTableTreeNode& InNode) const;
+
+	virtual void InitFilterConfigurator(FFilterConfigurator& InOutFilterConfigurator);
+	virtual void UpdateFilterContext(const FFilterConfigurator& InFilterConfigurator, const FTableTreeNode& InNode) const;
+
+	virtual TSharedRef<SWidget> ConstructFilterConfiguratorButton();
+	FReply FilterConfigurator_OnClicked();
+
+	////////////////////////////////////////////////////////////////////////////////////////////////////
+	// Hierarchy Filtering (Root->Children hierarchy --> Root->FilteredChildren hierarchy)
+
+	void InitHierarchyFiltering();
+
+	/**
+	 * Populates OutSearchStrings with the strings that should be used in searching.
+	 *
+	 * @param GroupOrStatNodePtr - the group and stat node to get a text description from.
+	 * @param OutSearchStrings   - an array of strings to use in searching.
+	 */
+	static void HandleItemToStringArray(const FTableTreeNodePtr& GroupOrStatNodePtr, TArray<FString>& OutSearchStrings);
+
+	void OnHierarchyFilteringChanged();
+	bool ScheduleHierarchyFilteringAsyncOperationIfNeeded();
+	void ScheduleHierarchyFilteringAsyncOperation();
+	FGraphEventRef StartHierarchyFilteringTask(FGraphEventRef Prerequisite = nullptr);
+
+	void ApplyHierarchyFiltering();
+	void ApplyEmptyHierarchyFilteringRec(FTableTreeNodePtr NodePtr);
+	bool ApplyHierarchyFilteringRec(FTableTreeNodePtr NodePtr);
+
+	/** Set all the nodes belonging to a subtree as visible. Returns true if the caller node should be expanded. */
+	bool MakeSubtreeVisible(FTableTreeNodePtr NodePtr, bool bFilterIsEmpty);
+
+	virtual TSharedRef<SWidget> ConstructSearchBox();
+	void SearchBox_OnTextChanged(const FText& InFilterText);
+
+	////////////////////////////////////////////////////////////////////////////////////////////////////
+	// Grouping (FilteredNodes --> Root->Children hierarchy)
+
+	void CreateGroupings();
+	virtual void InternalCreateGroupings();
+
+	void OnGroupingChanged();
+	bool ScheduleGroupingAsyncOperationIfNeeded();
+	void ScheduleGroupingAsyncOperation();
+	FGraphEventRef StartGroupingTask(FGraphEventRef Prerequisite = nullptr);
+	void ApplyGrouping();
+	void CreateGroups(const TArray<TSharedPtr<FTreeNodeGrouping>>& Groupings);
+	void GroupNodesRec(const TArray<FTableTreeNodePtr>& Nodes, FTableTreeNode& ParentGroup, int32 GroupingDepth, const TArray<TSharedPtr<FTreeNodeGrouping>>& Groupings);
+
+	void RebuildGroupingCrumbs();
+	void OnGroupingCrumbClicked(const TSharedPtr<FTreeNodeGrouping>& InEntry);
+	void BuildGroupingSubMenu_Change(FMenuBuilder& MenuBuilder, const TSharedPtr<FTreeNodeGrouping> CrumbGrouping);
+	void BuildGroupingSubMenu_Add(FMenuBuilder& MenuBuilder, const TSharedPtr<FTreeNodeGrouping> CrumbGrouping);
+	TSharedRef<SWidget> GetGroupingCrumbMenuContent(const TSharedPtr<FTreeNodeGrouping>& CrumbGrouping);
+
+	void PreChangeGroupings();
+	void PostChangeGroupings();
+	int32 GetGroupingDepth(const TSharedPtr<FTreeNodeGrouping>& Grouping) const;
+
+	void GroupingCrumbMenu_Reset_Execute();
+	void GroupingCrumbMenu_Remove_Execute(const TSharedPtr<FTreeNodeGrouping> Grouping);
+	void GroupingCrumbMenu_MoveLeft_Execute(const TSharedPtr<FTreeNodeGrouping> Grouping);
+	void GroupingCrumbMenu_MoveRight_Execute(const TSharedPtr<FTreeNodeGrouping> Grouping);
+	void GroupingCrumbMenu_Change_Execute(const TSharedPtr<FTreeNodeGrouping> OldGrouping, const TSharedPtr<FTreeNodeGrouping> NewGrouping);
+	bool GroupingCrumbMenu_Change_CanExecute(const TSharedPtr<FTreeNodeGrouping> OldGrouping, const TSharedPtr<FTreeNodeGrouping> NewGrouping) const;
+	void GroupingCrumbMenu_Add_Execute(const TSharedPtr<FTreeNodeGrouping> Grouping, const TSharedPtr<FTreeNodeGrouping> AfterGrouping);
+	bool GroupingCrumbMenu_Add_CanExecute(const TSharedPtr<FTreeNodeGrouping> Grouping, const TSharedPtr<FTreeNodeGrouping> AfterGrouping) const;
+
+	////////////////////////////////////////////////////////////////////////////////////////////////////
+	// Aggregation (Root->Children hierarchy)
+
+	static void UpdateCStringSameValueAggregationSingleNode(const FTableColumn& InColumn, FTableTreeNode& GroupNode);
+	static void UpdateCStringSameValueAggregationRec(const FTableColumn& InColumn, FTableTreeNode& GroupNode);
+
+	template<typename T, bool bSetInitialValue, bool bIsRercursive>
+	static void UpdateAggregation(const FTableColumn& InColumn, FTableTreeNode& InOutGroupNode, const T InitialAggregatedValue, TFunctionRef<T(T, const FTableCellValue&)> ValueGetterFunc);
+
+	template<bool bIsRercursive>
+	static void UpdateAggregatedValues(TSharedPtr<FTable> InTable, FTableTreeNode& InOutGroupNode);
+
+	void UpdateAggregatedValuesSingleNode(FTableTreeNode& GroupNode);
+	void UpdateAggregatedValuesRec(FTableTreeNode& GroupNode);
+
+	////////////////////////////////////////////////////////////////////////////////////////////////////
+	// Sorting (Root->Children hierarchy)
+
+	static const EColumnSortMode::Type GetDefaultColumnSortMode();
+	static const FName GetDefaultColumnBeingSorted();
+
+	void CreateSortings();
+	void UpdateCurrentSortingByColumn();
+
+	void OnSortingChanged();
+	bool ScheduleSortingAsyncOperationIfNeeded();
+	void ScheduleSortingAsyncOperation();
+	FGraphEventRef StartSortingTask(FGraphEventRef Prerequisite = nullptr);
+	void ApplySorting();
+	void SortTreeNodes(ITableCellValueSorter* InSorter, EColumnSortMode::Type InColumnSortMode);
+	void SortTreeNodesRec(FTableTreeNode& GroupNode, const ITableCellValueSorter& Sorter, EColumnSortMode::Type InColumnSortMode);
+
+	EColumnSortMode::Type GetSortModeForColumn(const FName ColumnId) const;
+	void SetSortModeForColumn(const FName& ColumnId, EColumnSortMode::Type SortMode);
+	void OnSortModeChanged(const EColumnSortPriority::Type SortPriority, const FName& ColumnId, const EColumnSortMode::Type SortMode);
+
+	////////////////////////////////////////////////////////////////////////////////////////////////////
+	// Sorting actions
+
+	// SortMode (HeaderMenu)
+	bool HeaderMenu_SortMode_IsChecked(const FName ColumnId, const EColumnSortMode::Type InSortMode);
+	bool HeaderMenu_SortMode_CanExecute(const FName ColumnId, const EColumnSortMode::Type InSortMode) const;
+	void HeaderMenu_SortMode_Execute(const FName ColumnId, const EColumnSortMode::Type InSortMode);
+
+	// SortMode (ContextMenu)
+	bool ContextMenu_SortMode_IsChecked(const EColumnSortMode::Type InSortMode);
+	bool ContextMenu_SortMode_CanExecute(const EColumnSortMode::Type InSortMode) const;
+	void ContextMenu_SortMode_Execute(const EColumnSortMode::Type InSortMode);
+
+	// SortByColumn (ContextMenu)
+	bool ContextMenu_SortByColumn_IsChecked(const FName ColumnId);
+	bool ContextMenu_SortByColumn_CanExecute(const FName ColumnId) const;
+	void ContextMenu_SortByColumn_Execute(const FName ColumnId);
+
+	////////////////////////////////////////////////////////////////////////////////////////////////////
+	// Column visibility actions
+
+	// ShowColumn
+	bool CanShowColumn(const FName ColumnId) const;
+	void ShowColumn(const FName ColumnId);
+	void ShowColumn(FTableColumn& Column);
+
+	// HideColumn
+	bool CanHideColumn(const FName ColumnId) const;
+	void HideColumn(const FName ColumnId);
+	void HideColumn(FTableColumn& Column);
+
+	// ToggleColumnVisibility
+	bool IsColumnVisible(const FName ColumnId);
+	bool CanToggleColumnVisibility(const FName ColumnId) const;
+	void ToggleColumnVisibility(const FName ColumnId);
+
+	// ShowAllColumns (ContextMenu)
+	bool ContextMenu_ShowAllColumns_CanExecute() const;
+	void ContextMenu_ShowAllColumns_Execute();
+
+	// ResetColumns (ContextMenu)
+	bool ContextMenu_ResetColumns_CanExecute() const;
+	void ContextMenu_ResetColumns_Execute();
+
+	// HideAllColumns (ContextMenu)
+	bool ContextMenu_HideAllColumns_CanExecute() const;
+	void ContextMenu_HideAllColumns_Execute();
+
+	////////////////////////////////////////////////////////////////////////////////////////////////////
+	//Async
+
+	virtual void OnPreAsyncUpdate();
+	virtual void OnPostAsyncUpdate();
+
+	void AddInProgressAsyncOperation(EAsyncOperationType InType) { EnumAddFlags(InProgressAsyncOperations, InType); }
+	bool HasInProgressAsyncOperation(EAsyncOperationType InType) const { return EnumHasAnyFlags(InProgressAsyncOperations, InType); }
+	void ClearInProgressAsyncOperations() { InProgressAsyncOperations = static_cast<EAsyncOperationType>(0); }
+
+	void StartPendingAsyncOperations();
+
+	void CancelCurrentAsyncOp();
+
+	////////////////////////////////////////////////////////////////////////////////////////////////////
+
+	void CountNumNodesPerDepthRec(FBaseTreeNode* InRoot, TArray<int32>& InOutNumNodesPerDepth, int32 InDepth, int32 InMaxDepth, int32 InMaxNodes) const;
+	void SetExpandValueForChildGroups(FBaseTreeNode* InRoot, int32 InMaxExpandedNodes, int32 MaxDepthToExpand, bool InValue);
+	void SetExpandValueForChildGroupsRec(FBaseTreeNode* InRoot, int32 InDepth, int32 InMaxDepth, bool InValue);
+
+	virtual void ExtendMenu(TSharedRef<FExtender> Extender) {}
+	virtual void ExtendMenu(FMenuBuilder& Menu) {}
+
+	typedef TFunctionRef<void(TArray<FBaseTreeNodePtr>& InNodes)> WriteToFileCallback;
+	void ExportToFileRec(const FBaseTreeNodePtr& InGroupNode, TArray<FBaseTreeNodePtr>& InNodes, bool bInExportCollapsed, bool InExportLeafs, WriteToFileCallback Callback);
+
+	FText GetTreeViewBannerText() const { return TreeViewBannerText; }
+	virtual void UpdateBannerText();
+
+	////////////////////////////////////////////////////////////////////////////////////////////////////
+	// Presets
+
+	virtual void InitAvailableViewPresets() {};
+	const TArray<TSharedRef<ITableTreeViewPreset>>* GetAvailableViewPresets() const { return &AvailableViewPresets; }
+	FReply OnApplyViewPreset(const ITableTreeViewPreset* InPreset);
+	void ApplyViewPreset(const ITableTreeViewPreset& InPreset);
+	void ApplyColumnConfig(const TArrayView<FTableColumnConfig>& InTableConfig);
+	void ViewPreset_OnSelectionChanged(TSharedPtr<ITableTreeViewPreset> InPreset, ESelectInfo::Type SelectInfo);
+	TSharedRef<SWidget> ViewPreset_OnGenerateWidget(TSharedRef<ITableTreeViewPreset> InPreset);
+	FText ViewPreset_GetSelectedText() const;
+	FText ViewPreset_GetSelectedToolTipText() const;
+
+	////////////////////////////////////////////////////////////////////////////////////////////////////
+
+	virtual void SearchForItem(TSharedPtr<FTableTaskCancellationToken> CancellationToken) {};
+
+	// Table data tasks should be tasks that operate read only operations on the data from the Table
+	// They should not operate on the tree nodes because they will run concurrently with the populated table UI.
+	template<typename T, typename... TArgs>
+	TSharedPtr<FTableTaskInfo> StartTableDataTask(TArgs&&... Args)
+	{
+		TSharedPtr<FTableTaskInfo> Info = MakeShared<FTableTaskInfo>();
+		Info->CancellationToken = MakeShared<FTableTaskCancellationToken>();
+		Info->Event = TGraphTask<T>::CreateTask().ConstructAndDispatchWhenReady(Info->CancellationToken, Forward<TArgs>(Args)...);
+		DataTaskInfos.Add(Info);
+		return Info;
+	}
+
+	void StopAllTableDataTasks(bool bWait = true);
+
+protected:
+	/** Table view model. */
+	TSharedPtr<FTable> Table;
+
+	TSharedPtr<FUICommandList> CommandList;
+
+	//////////////////////////////////////////////////
+	// Widget
+
+	/** The child STreeView widget. */
+	TSharedPtr<STreeView<FTableTreeNodePtr>> TreeView;
+
+	/** Holds the tree view header row widget which display all columns in the tree view. */
+	TSharedPtr<SHeaderRow> TreeViewHeaderRow;
+
+	/** External scrollbar used to synchronize tree view position. */
+	TSharedPtr<SScrollBar> ExternalScrollbar;
+
+	//////////////////////////////////////////////////
+	// Hovered Column, Hovered Tree Node
+
+	/** Name of the column currently being hovered by the mouse. */
+	FName HoveredColumnId;
+
+	/** A shared pointer to the tree node currently being hovered by the mouse. */
+	FTableTreeNodePtr HoveredNodePtr;
+
+	/** Name of the tree node that should be drawn as highlighted. */
+	FName HighlightedNodeName;
+
+	//////////////////////////////////////////////////
+	// Tree Nodes
+
+	static const FName RootNodeName;
+
+	/** The root node of the tree. It is invisible; only its children are displayed in the tree view. */
+	FTableTreeNodePtr Root;
+
+	/** Table row nodes. Each node corresponds to a table row. Index in this array corresponds to RowIndex in source table. */
+	TArray<FTableTreeNodePtr> TableRowNodes;
+
+	/** Filtered table nodes. These are the nodes from the TableRowNodes array, after applying the node filtering. */
+	TArray<FTableTreeNodePtr> FilteredNodes;
+
+	/** A pointer to the filtered table nodes. If node filtering is empty, this points directly to TableRowNodes. */
+	TArray<FTableTreeNodePtr>* FilteredNodesPtr = nullptr;
+
+	/** A filtered array of group nodes to be displayed in the tree widget. */
+	TArray<FTableTreeNodePtr> FilteredGroupNodes;
+
+	/** Currently expanded group nodes. */
+	TSet<FTableTreeNodePtr> ExpandedNodes;
+
+	/** If true, the expanded nodes have been saved before applying a text filter. Only used when bRunInAsyncMode==false. */
+	bool bExpansionSaved = false;
+
+	int32 MaxNodesToAutoExpand = 1000;
+	int32 MaxDepthToAutoExpand = 4;
+	int32 MaxNodesToExpand = 1000000;
+	int32 MaxDepthToExpand = 100;
+
+	//////////////////////////////////////////////////
+	// Search box & the hierarchy filtering
+
+	/** The search box widget used to filter items displayed in the stats and groups tree. */
+	TSharedPtr<SSearchBox> SearchBox;
+
+	/** The filter collection. */
+	TSharedPtr<FTableTreeNodeFilterCollection> Filters;
+
+	/** The text based filter. */
+	TSharedPtr<FTableTreeNodeTextFilter> TextFilter;
+
+	/** The text based filter actually used in the Hierarchy Filtering async task. */
+	TSharedPtr<FTableTreeNodeTextFilter> CurrentAsyncOpTextFilter;
+
+	//////////////////////////////////////////////////
+	// Node filtering
+
+	TSharedPtr<FFilterConfigurator> FilterConfigurator;
+	FDelegateHandle OnFilterChangesCommittedHandle;
+
+	/** The filter configurator actually used in the Hierarchy Filtering async task. */
+	FFilterConfigurator* CurrentAsyncOpFilterConfigurator = nullptr;
+
+	mutable FFilterContext FilterContext;
+
+	//////////////////////////////////////////////////
+	// Grouping
+
+	TArray<TSharedPtr<FTreeNodeGrouping>> AvailableGroupings;
+
+	/** How we group the tree nodes? */
+	TArray<TSharedPtr<FTreeNodeGrouping>> CurrentGroupings;
+
+	TSharedPtr<SBreadcrumbTrail<TSharedPtr<FTreeNodeGrouping>>> GroupingBreadcrumbTrail;
+
+	/** The groupings actually used in the Grouping async task. */
+	TArray<TSharedPtr<FTreeNodeGrouping>> CurrentAsyncOpGroupings;
+
+	//////////////////////////////////////////////////
+	// Sorting
+
+	/** All available sorters. */
+	TArray<TSharedPtr<ITableCellValueSorter>> AvailableSorters;
+
+	/** Current sorter. It is nullptr if sorting is disabled. */
+	TSharedPtr<ITableCellValueSorter> CurrentSorter;
+
+	/** Name of the column currently being sorted. Can be NAME_None if sorting is disabled (CurrentSorting == nullptr) or if a complex sorting is used (CurrentSorting != nullptr). */
+	FName ColumnBeingSorted;
+
+	/** How we sort the nodes? Ascending or Descending. */
+	EColumnSortMode::Type ColumnSortMode = EColumnSortMode::None;
+
+	/** The sorter actually used in the Sorting async task. */
+	ITableCellValueSorter* CurrentAsyncOpSorter = nullptr;
+
+	/** The sort mode actually used in the Sorting async task. */
+	EColumnSortMode::Type CurrentAsyncOpColumnSortMode;
+
+	//////////////////////////////////////////////////
+	// Async Operations
+
+	bool bRunInAsyncMode = false;
+	bool bIsUpdateRunning = false;
+	bool bIsCloseScheduled = false;
+
+	TArray<FTableTreeNodePtr> DummyGroupNodes;
+	FGraphEventRef InProgressAsyncOperationEvent;
+	FGraphEventRef AsyncCompleteTaskEvent;
+	EAsyncOperationType InProgressAsyncOperations = static_cast<EAsyncOperationType>(0);
+	TSharedPtr<class SAsyncOperationStatus> AsyncOperationStatus;
+	FStopwatch AsyncUpdateStopwatch;
+	FAsyncOperationProgress AsyncOperationProgress;
+	FGraphEventRef DispatchEvent;
+	TArray<TSharedPtr<FTableTaskInfo>> DataTaskInfos;
+	TArray<FTableTreeNodePtr> NodesToExpand;
+
+	//////////////////////////////////////////////////
+
+	FText TreeViewBannerText;
+
+	TArray<TSharedRef<ITableTreeViewPreset>> AvailableViewPresets;
+	TSharedPtr<ITableTreeViewPreset> SelectedViewPreset;
+	TSharedPtr<SComboBox<TSharedRef<ITableTreeViewPreset>>> PresetComboBox;
+
+	//////////////////////////////////////////////////
+	// Logging
+
+	/** A log listing name to be used for any errors or warnings. */
+	FName LogListingName;
+};
+
+////////////////////////////////////////////////////////////////////////////////////////////////////
+
+class TRACEINSIGHTSCORE_API FTableTreeViewNodeFilteringAsyncTask
+{
+public:
+	FTableTreeViewNodeFilteringAsyncTask(STableTreeView* InPtr)
+		: TableTreeViewPtr(InPtr)
+	{
+	}
+
+	FORCEINLINE TStatId GetStatId() const { RETURN_QUICK_DECLARE_CYCLE_STAT(FTableTreeViewNodeFilteringAsyncTask, STATGROUP_TaskGraphTasks); }
+	ENamedThreads::Type GetDesiredThread() { return ENamedThreads::Type::AnyThread; }
+	static ESubsequentsMode::Type GetSubsequentsMode() { return ESubsequentsMode::TrackSubsequents; }
+
+	void DoTask(ENamedThreads::Type CurrentThread, const FGraphEventRef& MyCompletionGraphEvent)
+	{
+		if (TableTreeViewPtr)
+		{
+			TableTreeViewPtr->ApplyNodeFiltering();
+		}
+	}
+
+private:
+	STableTreeView* TableTreeViewPtr = nullptr;
+};
+
+////////////////////////////////////////////////////////////////////////////////////////////////////
+
+class TRACEINSIGHTSCORE_API FTableTreeViewHierarchyFilteringAsyncTask
+{
+public:
+	FTableTreeViewHierarchyFilteringAsyncTask(STableTreeView* InPtr)
+		: TableTreeViewPtr(InPtr)
+	{
+	}
+
+	FORCEINLINE TStatId GetStatId() const { RETURN_QUICK_DECLARE_CYCLE_STAT(FTableTreeViewHierarchyFilteringAsyncTask, STATGROUP_TaskGraphTasks); }
+	ENamedThreads::Type GetDesiredThread() { return ENamedThreads::Type::AnyThread; }
+	static ESubsequentsMode::Type GetSubsequentsMode() { return ESubsequentsMode::TrackSubsequents; }
+
+	void DoTask(ENamedThreads::Type CurrentThread, const FGraphEventRef& MyCompletionGraphEvent)
+	{
+		if (TableTreeViewPtr)
+		{
+			TableTreeViewPtr->ApplyHierarchyFiltering();
+		}
+	}
+
+private:
+	STableTreeView* TableTreeViewPtr = nullptr;
+};
+
+////////////////////////////////////////////////////////////////////////////////////////////////////
+
+class TRACEINSIGHTSCORE_API FTableTreeViewSortingAsyncTask
+{
+public:
+	FTableTreeViewSortingAsyncTask(STableTreeView* InPtr, ITableCellValueSorter* InSorter, EColumnSortMode::Type InColumnSortMode)
+		: TableTreeViewPtr(InPtr)
+		, Sorter(InSorter)
+		, ColumnSortMode(InColumnSortMode)
+	{
+	}
+
+	FORCEINLINE TStatId GetStatId() const { RETURN_QUICK_DECLARE_CYCLE_STAT(FTableTreeViewSortingAsyncTask, STATGROUP_TaskGraphTasks); }
+	ENamedThreads::Type GetDesiredThread() { return ENamedThreads::Type::AnyThread; }
+	static ESubsequentsMode::Type GetSubsequentsMode() { return ESubsequentsMode::TrackSubsequents; }
+
+	void DoTask(ENamedThreads::Type CurrentThread, const FGraphEventRef& MyCompletionGraphEvent)
+	{
+		if (TableTreeViewPtr)
+		{
+			TableTreeViewPtr->SortTreeNodes(Sorter, ColumnSortMode);
+		}
+	}
+
+private:
+	STableTreeView* TableTreeViewPtr = nullptr;
+	ITableCellValueSorter* Sorter = nullptr;
+	EColumnSortMode::Type ColumnSortMode = EColumnSortMode::Type::None;
+};
+
+////////////////////////////////////////////////////////////////////////////////////////////////////
+
+class TRACEINSIGHTSCORE_API FTableTreeViewGroupingAsyncTask
+{
+public:
+	FTableTreeViewGroupingAsyncTask(STableTreeView* InPtr, TArray<TSharedPtr<FTreeNodeGrouping>>* InGroupings)
+		: TableTreeViewPtr(InPtr)
+		, Groupings(InGroupings)
+	{
+	}
+
+	FORCEINLINE TStatId GetStatId() const { RETURN_QUICK_DECLARE_CYCLE_STAT(FTableTreeViewGroupingAsyncTask, STATGROUP_TaskGraphTasks); }
+	ENamedThreads::Type GetDesiredThread() { return ENamedThreads::Type::AnyThread; }
+	static ESubsequentsMode::Type GetSubsequentsMode() { return ESubsequentsMode::TrackSubsequents; }
+
+	void DoTask(ENamedThreads::Type CurrentThread, const FGraphEventRef& MyCompletionGraphEvent)
+	{
+		if (TableTreeViewPtr)
+		{
+			TableTreeViewPtr->CreateGroups(*Groupings);
+		}
+	}
+
+private:
+	STableTreeView* TableTreeViewPtr = nullptr;
+	TArray<TSharedPtr<FTreeNodeGrouping>>* Groupings = nullptr;
+};
+
+////////////////////////////////////////////////////////////////////////////////////////////////////
+
+class TRACEINSIGHTSCORE_API FSearchForItemToSelectTask
+{
+public:
+	FSearchForItemToSelectTask(TSharedPtr<FTableTaskCancellationToken> InToken, TSharedPtr<STableTreeView> InPtr)
+		: CancellationToken(InToken)
+		, TableTreeViewPtr(InPtr)
+	{
+	}
+
+	FORCEINLINE TStatId GetStatId() const { RETURN_QUICK_DECLARE_CYCLE_STAT(FSearchForItemToSelectTask, STATGROUP_TaskGraphTasks); }
+	ENamedThreads::Type GetDesiredThread() { return ENamedThreads::Type::AnyThread; }
+	static ESubsequentsMode::Type GetSubsequentsMode() { return ESubsequentsMode::TrackSubsequents; }
+
+	void DoTask(ENamedThreads::Type CurrentThread, const FGraphEventRef& MyCompletionGraphEvent)
+	{
+		if (!CancellationToken.IsValid() || !CancellationToken->ShouldCancel())
+		{
+			if (TableTreeViewPtr.IsValid())
+			{
+				TableTreeViewPtr->SearchForItem(CancellationToken);
+			}
+		}
+	}
+
+private:
+	TSharedPtr<FTableTaskCancellationToken> CancellationToken;
+	TSharedPtr<STableTreeView> TableTreeViewPtr;
+};
+
+////////////////////////////////////////////////////////////////////////////////////////////////////
+
+class TRACEINSIGHTSCORE_API FSelectNodeByTableRowIndexTask
+{
+public:
+	FSelectNodeByTableRowIndexTask(TSharedPtr<FTableTaskCancellationToken> InToken, TSharedPtr<STableTreeView> InPtr, uint32 InRowIndex)
+		: CancellationToken(InToken)
+		, TableTreeViewPtr(InPtr)
+		, RowIndex(InRowIndex)
+	{
+	}
+
+	FORCEINLINE TStatId GetStatId() const { RETURN_QUICK_DECLARE_CYCLE_STAT(FSelectNodeByTableRowIndexTask, STATGROUP_TaskGraphTasks); }
+	ENamedThreads::Type GetDesiredThread() { return ENamedThreads::Type::GameThread; }
+	static ESubsequentsMode::Type GetSubsequentsMode() { return ESubsequentsMode::TrackSubsequents; }
+
+	void DoTask(ENamedThreads::Type CurrentThread, const FGraphEventRef& MyCompletionGraphEvent)
+	{
+		if (!CancellationToken.IsValid() || !CancellationToken->ShouldCancel())
+		{
+			if (TableTreeViewPtr.IsValid())
+			{
+				TableTreeViewPtr->SelectNodeByTableRowIndex(RowIndex);
+			}
+		}
+	}
+
+private:
+	TSharedPtr< FTableTaskCancellationToken> CancellationToken;
+	TSharedPtr<STableTreeView> TableTreeViewPtr;
+	uint32 RowIndex = 0;
+};
+
+////////////////////////////////////////////////////////////////////////////////////////////////////
+
+} // namespace UE::Insights
