[Insights]
  - Fixed initialization of FFilterService to be done automatically by the TraceInsightsCore module.
  - Removed FCustomFilter (was identical with FFilterWithSuggestions).
  - Removed some unnecessary TRACEINSIGHTSCORE_API for public filter classes.

#rnx

diff --git a/Engine/Plugins/Editor/AssetManagerEditor/Source/AssetManagerEditor/Private/AssetManagerEditorModule.cpp b/Engine/Plugins/Editor/AssetManagerEditor/Source/AssetManagerEditor/Private/AssetManagerEditorModule.cpp
--- a/Engine/Plugins/Editor/AssetManagerEditor/Source/AssetManagerEditor/Private/AssetManagerEditorModule.cpp
+++ b/Engine/Plugins/Editor/AssetManagerEditor/Source/AssetManagerEditor/Private/AssetManagerEditorModule.cpp
@@ -66,7 +66,6 @@
 #include "ContentBrowserMenuContexts.h"
 #include "IContentBrowserDataModule.h"
 #include "ContentBrowserDataSubsystem.h"
-#include "InsightsCore/Filter/ViewModels/Filters.h"
 #include "TreeView/AssetTable.h"
 #include "TreeView/SAssetTableTreeView.h"
 
@@ -485,8 +484,6 @@ void FAssetManagerEditorModule::StartupModule()
 
 	if (GIsEditor && !IsRunningCommandlet())
 	{
-		UE::Insights::FFilterService::Initialize();
-
 		AuditCmds.Add(IConsoleManager::Get().RegisterConsoleCommand(
 			TEXT("AssetManager.AssetAudit"),
 			TEXT("Dumps statistics about assets to the log."),
@@ -693,8 +690,6 @@ void FAssetManagerEditorModule::ShutdownModule()
 		// Cleanup tool menus
 		UToolMenus::UnRegisterStartupCallback(this);
 		UToolMenus::UnregisterOwner(this);
-
-		UE::Insights::FFilterService::Shutdown();
 	}
 }
 
diff --git a/Engine/Source/Developer/TraceInsights/Private/Insights/InsightsManager.cpp b/Engine/Source/Developer/TraceInsights/Private/Insights/InsightsManager.cpp
--- a/Engine/Source/Developer/TraceInsights/Private/Insights/InsightsManager.cpp
+++ b/Engine/Source/Developer/TraceInsights/Private/Insights/InsightsManager.cpp
@@ -28,7 +28,6 @@
 
 // TraceInsightsCore
 #include "InsightsCore/Common/TimeUtils.h"
-#include "InsightsCore/Filter/ViewModels/Filters.h"
 
 // TraceInsights
 #include "Insights/Common/InsightsMenuBuilder.h"
@@ -130,8 +129,6 @@ void FInsightsManager::Initialize(IUnrealInsightsModule& InsightsModule)
 
 	InsightsMenuBuilder = MakeShared<FInsightsMenuBuilder>();
 
-	UE::Insights::FFilterService::Initialize();
-
 	FMessageLogModule& MessageLogModule = FModuleManager::LoadModuleChecked<FMessageLogModule>("MessageLog");
 	MessageLogModule.RegisterLogListing(GetLogListingName(), LOCTEXT("UnrealInsights", "Unreal Insights"));
 	MessageLogModule.RegisterLogListing(AnalysisLogListingName, LOCTEXT("TraceAnalysis", "Trace Analysis"));
@@ -172,8 +169,6 @@ void FInsightsManager::Shutdown()
 		}
 	}
 
-	UE::Insights::FFilterService::Shutdown();
-
 	FInsightsManager::Instance.Reset();
 }
 
@@ -267,9 +262,11 @@ void FInsightsManager::UnregisterMajorTabs()
 {
 	FGlobalTabmanager::Get()->UnregisterNomadTabSpawner(FInsightsManagerTabs::SessionInfoTabId);
 
-	FGlobalTabmanager::Get()->UnregisterNomadTabSpawner(FInsightsManagerTabs::TraceStoreTabId);
-	FGlobalTabmanager::Get()->UnregisterNomadTabSpawner(FInsightsManagerTabs::ConnectionTabId);
+#if !WITH_EDITOR
 	FGlobalTabmanager::Get()->UnregisterNomadTabSpawner(FInsightsManagerTabs::LauncherTabId);
+	FGlobalTabmanager::Get()->UnregisterNomadTabSpawner(FInsightsManagerTabs::ConnectionTabId);
+	FGlobalTabmanager::Get()->UnregisterNomadTabSpawner(FInsightsManagerTabs::TraceStoreTabId);
+#endif
 }
 
 ////////////////////////////////////////////////////////////////////////////////////////////////////
diff --git a/Engine/Source/Developer/TraceInsights/Private/Insights/ViewModels/TimerFilters.cpp b/Engine/Source/Developer/TraceInsights/Private/Insights/ViewModels/TimerFilters.cpp
--- a/Engine/Source/Developer/TraceInsights/Private/Insights/ViewModels/TimerFilters.cpp
+++ b/Engine/Source/Developer/TraceInsights/Private/Insights/ViewModels/TimerFilters.cpp
@@ -106,7 +106,8 @@ TSharedRef<UE::Insights::FFilterState> FTimerNameFilterState::DeepCopy() const
 ////////////////////////////////////////////////////////////////////////////////////////////////////
 
 FTimerNameFilter::FTimerNameFilter()
-	: FCustomFilter(static_cast<int32>(UE::Insights::EFilterField::TimerName),
+	: FFilterWithSuggestions(
+		static_cast<int32>(UE::Insights::EFilterField::TimerName),
 		LOCTEXT("TimerName", "Timer Name"),
 		LOCTEXT("TimerName", "Timer Name"),
 		UE::Insights::EFilterDataType::Custom,
diff --git a/Engine/Source/Developer/TraceInsights/Private/Insights/ViewModels/TimerFilters.h b/Engine/Source/Developer/TraceInsights/Private/Insights/ViewModels/TimerFilters.h
--- a/Engine/Source/Developer/TraceInsights/Private/Insights/ViewModels/TimerFilters.h
+++ b/Engine/Source/Developer/TraceInsights/Private/Insights/ViewModels/TimerFilters.h
@@ -50,9 +50,9 @@ private:
 	TSet<uint32> TimerIds;
 };
 
-class FTimerNameFilter : public UE::Insights::FCustomFilter
+class FTimerNameFilter : public UE::Insights::FFilterWithSuggestions
 {
-	INSIGHTS_DECLARE_RTTI(FTimerNameFilter, UE::Insights::FCustomFilter)
+	INSIGHTS_DECLARE_RTTI(FTimerNameFilter, UE::Insights::FFilterWithSuggestions)
 
 public:
 	FTimerNameFilter();
diff --git a/Engine/Source/Developer/TraceInsightsCore/Private/InsightsCore/Filter/ViewModels/Filters.cpp b/Engine/Source/Developer/TraceInsightsCore/Private/InsightsCore/Filter/ViewModels/Filters.cpp
--- a/Engine/Source/Developer/TraceInsightsCore/Private/InsightsCore/Filter/ViewModels/Filters.cpp
+++ b/Engine/Source/Developer/TraceInsightsCore/Private/InsightsCore/Filter/ViewModels/Filters.cpp
@@ -17,7 +17,6 @@ namespace UE::Insights
 INSIGHTS_IMPLEMENT_RTTI(FFilterState)
 INSIGHTS_IMPLEMENT_RTTI(FFilter)
 INSIGHTS_IMPLEMENT_RTTI(FFilterWithSuggestions)
-INSIGHTS_IMPLEMENT_RTTI(FCustomFilter)
 
 ////////////////////////////////////////////////////////////////////////////////////////////////////
 // FFilterStorage
@@ -70,6 +69,7 @@ FName const FFilterService::FilterConfiguratorTabId(TEXT("FilterConfigurator"));
 
 void FFilterService::Initialize()
 {
+	check(!Instance.IsValid());
 	Instance = MakeShared<FFilterService>();
 }
 
diff --git a/Engine/Source/Developer/TraceInsightsCore/Private/InsightsCore/TraceInsightsCoreModule.cpp b/Engine/Source/Developer/TraceInsightsCore/Private/InsightsCore/TraceInsightsCoreModule.cpp
--- a/Engine/Source/Developer/TraceInsightsCore/Private/InsightsCore/TraceInsightsCoreModule.cpp
+++ b/Engine/Source/Developer/TraceInsightsCore/Private/InsightsCore/TraceInsightsCoreModule.cpp
@@ -5,15 +5,18 @@
 #include "Modules/ModuleManager.h"
 
 #include "InsightsCore/Common/InsightsCoreStyle.h"
+#include "InsightsCore/Filter/ViewModels/Filters.h"
 
 IMPLEMENT_MODULE(FTraceInsightsCoreModule, TraceInsightsCore);
 
 void FTraceInsightsCoreModule::StartupModule()
 {
 	UE::Insights::FInsightsCoreStyle::Initialize();
+	UE::Insights::FFilterService::Initialize();
 }
 
 void FTraceInsightsCoreModule::ShutdownModule()
 {
+	UE::Insights::FFilterService::Shutdown();
 	UE::Insights::FInsightsCoreStyle::Shutdown();
 }
diff --git a/Engine/Source/Developer/TraceInsightsCore/Public/InsightsCore/Filter/ViewModels/Filters.h b/Engine/Source/Developer/TraceInsightsCore/Public/InsightsCore/Filter/ViewModels/Filters.h
--- a/Engine/Source/Developer/TraceInsightsCore/Public/InsightsCore/Filter/ViewModels/Filters.h
+++ b/Engine/Source/Developer/TraceInsightsCore/Public/InsightsCore/Filter/ViewModels/Filters.h
@@ -23,7 +23,7 @@ namespace UE::Insights
 
 ////////////////////////////////////////////////////////////////////////////////////////////////////
 
-enum class TRACEINSIGHTSCORE_API EFilterDataType : uint32
+enum class EFilterDataType : uint32
 {
 	Int64,
 	Double,
@@ -35,7 +35,7 @@ enum class TRACEINSIGHTSCORE_API EFilterDataType : uint32
 
 ////////////////////////////////////////////////////////////////////////////////////////////////////
 
-enum class TRACEINSIGHTSCORE_API EFilterOperator : uint8
+enum class EFilterOperator : uint8
 {
 	Eq = 0, // Equals
 	NotEq = 1, // Not Equals
@@ -49,7 +49,7 @@ enum class TRACEINSIGHTSCORE_API EFilterOperator : uint8
 
 ////////////////////////////////////////////////////////////////////////////////////////////////////
 
-enum class TRACEINSIGHTSCORE_API EFilterField : int32
+enum class EFilterField : int32
 {
 	StartTime = 0,
 	EndTime = 1,
@@ -64,7 +64,7 @@ enum class TRACEINSIGHTSCORE_API EFilterField : int32
 
 ////////////////////////////////////////////////////////////////////////////////////////////////////
 
-class TRACEINSIGHTSCORE_API IFilterOperator
+class IFilterOperator
 {
 public:
 	virtual EFilterOperator GetKey() const = 0;
@@ -104,7 +104,7 @@ private:
 
 ////////////////////////////////////////////////////////////////////////////////////////////////////
 
-enum class TRACEINSIGHTSCORE_API EFilterGroupOperator
+enum class EFilterGroupOperator
 {
 	And = 0,
 	Or = 1,
@@ -112,7 +112,7 @@ enum class TRACEINSIGHTSCORE_API EFilterGroupOperator
 
 ////////////////////////////////////////////////////////////////////////////////////////////////////
 
-class TRACEINSIGHTSCORE_API FFilterGroupOperator
+class FFilterGroupOperator
 {
 public:
 	FFilterGroupOperator(EFilterGroupOperator InType, FText InName, FText InDesc)
@@ -134,7 +134,7 @@ private:
 
 ////////////////////////////////////////////////////////////////////////////////////////////////////
 
-class TRACEINSIGHTSCORE_API IFilterValueConverter
+class IFilterValueConverter
 {
 public:
 	virtual bool Convert(const FString& Input, int64& Output, FText& OutError) const { unimplemented(); return false; }
@@ -145,7 +145,7 @@ public:
 
 ////////////////////////////////////////////////////////////////////////////////////////////////////
 
-class TRACEINSIGHTSCORE_API FFilterContext
+class FFilterContext
 {
 public:
 	typedef TVariant<double, int64, FString> ContextData;
@@ -200,6 +200,7 @@ class FFilter;
 class TRACEINSIGHTSCORE_API FFilterState
 {
 	INSIGHTS_DECLARE_RTTI_BASE(FFilterState)
+
 public:
 	FFilterState(TSharedRef<FFilter> InFilter)
 		: Filter(InFilter)
@@ -358,29 +359,11 @@ public:
 
 private:
 	static const FName FilterConfiguratorTabId;
-
 	static TSharedPtr<FFilterService> Instance;
 	FFilterStorage FilterStorage;
-
 	TSharedPtr<class SAdvancedFilter> PendingWidget;
 };
 
 ////////////////////////////////////////////////////////////////////////////////////////////////////
 
-class TRACEINSIGHTSCORE_API FCustomFilter : public FFilterWithSuggestions
-{
-	INSIGHTS_DECLARE_RTTI(FCustomFilter, FFilterWithSuggestions)
-public:
-	FCustomFilter(int32 InKey, FText InName, FText InDesc, EFilterDataType InDataType, TSharedPtr<IFilterValueConverter> InConverter, SupportedOperatorsArrayPtr InSupportedOperators)
-		: FFilterWithSuggestions(InKey, InName, InDesc, InDataType, InConverter, InSupportedOperators)
-	{
-	}
-
-	virtual ~FCustomFilter()
-	{
-	}
-};
-
-////////////////////////////////////////////////////////////////////////////////////////////////////
-
 } // namespace UE::Insights
