[insights-core] remove now-obsolete FString::Printf fallback

diff --git a/Engine/Source/Developer/TraceInsightsCore/Private/InsightsCore/Common/TimeUtils.cpp b/Engine/Source/Developer/TraceInsightsCore/Private/InsightsCore/Common/TimeUtils.cpp
--- a/Engine/Source/Developer/TraceInsightsCore/Private/InsightsCore/Common/TimeUtils.cpp
+++ b/Engine/Source/Developer/TraceInsightsCore/Private/InsightsCore/Common/TimeUtils.cpp
@@ -26,14 +26,7 @@ FString FormatTimeValue(const double Duration, const int32 NumDigits)
 		return FString::Printf(TEXT("%.0f"), Duration);
 	}
 
-#if !PLATFORM_USE_GENERIC_STRING_IMPLEMENTATION
 	FString Str = FString::Printf(TEXT("%.*f"), NumDigits, Duration);
-#else
-	// proper resolution is tracked as UE-79534
-	TCHAR FormatString[32];
-	FCString::Snprintf(FormatString, sizeof(FormatString), TEXT("%%.%df"), NumDigits);
-	FString Str = FString::Printf(FormatString, Duration);
-#endif
 
 	int32 Index = Str.Len() - 1;
 	while (Index > 0 && Str[Index] == '0') --Index;
