[Insights] Added more customization hooks to STableTreeView:
* toolbar customization: ConstructFilterToolbar(), ConstructExtraToolbar(), ConstructViewPreset()
* support for custom filtering: HasCustomNodeFilter() + FilterNodeCustom()

diff --git a/Engine/Source/Developer/TraceInsightsCore/Private/InsightsCore/Table/Widgets/STableTreeView.cpp b/Engine/Source/Developer/TraceInsightsCore/Private/InsightsCore/Table/Widgets/STableTreeView.cpp
--- a/Engine/Source/Developer/TraceInsightsCore/Private/InsightsCore/Table/Widgets/STableTreeView.cpp
+++ b/Engine/Source/Developer/TraceInsightsCore/Private/InsightsCore/Table/Widgets/STableTreeView.cpp
@@ -377,66 +377,103 @@ END_SLATE_FUNCTION_BUILD_OPTIMIZATION
 ////////////////////////////////////////////////////////////////////////////////////////////////////
 
 BEGIN_SLATE_FUNCTION_BUILD_OPTIMIZATION
-void STableTreeView::ConstructHeaderArea(TSharedRef<SVerticalBox> InWidgetContent)
+TSharedPtr<SWidget> STableTreeView::ConstructFilterToolbar()
 {
-	InWidgetContent->AddSlot()
-	.VAlign(VAlign_Center)
-	.AutoHeight()
-	.Padding(2.0f)
-	[
-		SNew(SHorizontalBox)
+	TSharedPtr<SHorizontalBox> Box = SNew(SHorizontalBox);
 
-		+ SHorizontalBox::Slot()
+	Box->AddSlot()
 		.FillWidth(1.0f)
 		.VAlign(VAlign_Center)
 		[
 			ConstructSearchBox()
-		]
+		];
 
-		+ SHorizontalBox::Slot()
+	Box->AddSlot()
 		.AutoWidth()
 		.Padding(4.0f, 0.0f, 0.0f, 0.0f)
 		.VAlign(VAlign_Center)
 		[
 			ConstructFilterConfiguratorButton()
-		]
-	];
+		];
 
-	TSharedPtr<SWidget> ToolbarWidget = ConstructToolbar();
-	if (ToolbarWidget)
+	return Box;
+}
+
+////////////////////////////////////////////////////////////////////////////////////////////////////
+
+BEGIN_SLATE_FUNCTION_BUILD_OPTIMIZATION
+void STableTreeView::ConstructHeaderArea(TSharedRef<SVerticalBox> InHostBox)
+{
+	TSharedPtr<SWidget> FilterToolbarWidget = ConstructFilterToolbar();
+	if (FilterToolbarWidget)
 	{
-		InWidgetContent->AddSlot()
-		.VAlign(VAlign_Center)
-		.AutoHeight()
-		.Padding(2.0f)
-		[
-			SNew(SHorizontalBox)
+		InHostBox->AddSlot()
+			.VAlign(VAlign_Center)
+			.AutoHeight()
+			.Padding(2.0f)
+			[
+				FilterToolbarWidget.ToSharedRef()
+			];
+	}
 
-			+ SHorizontalBox::Slot()
-			.FillWidth(1.0f)
+	TSharedPtr<SWidget> HierarchyBreadcrumbTrailWidget = ConstructHierarchyBreadcrumbTrail();
+	TSharedPtr<SWidget> ToolbarWidget = ConstructToolbar();
+	if (HierarchyBreadcrumbTrailWidget && ToolbarWidget)
+	{
+		InHostBox->AddSlot()
 			.VAlign(VAlign_Center)
+			.AutoHeight()
+			.Padding(2.0f)
 			[
-				ConstructHierarchyBreadcrumbTrail()
-			]
+				SNew(SHorizontalBox)
 
-			+ SHorizontalBox::Slot()
-			.AutoWidth()
-			.Padding(4.0f, 0.0f, 0.0f, 0.0f)
+				+ SHorizontalBox::Slot()
+				.FillWidth(1.0f)
+				.VAlign(VAlign_Center)
+				[
+					HierarchyBreadcrumbTrailWidget.ToSharedRef()
+				]
+
+				+ SHorizontalBox::Slot()
+				.AutoWidth()
+				.Padding(4.0f, 0.0f, 0.0f, 0.0f)
+				.VAlign(VAlign_Center)
+				[
+					ToolbarWidget.ToSharedRef()
+				]
+			];
+	}
+	else if (HierarchyBreadcrumbTrailWidget)
+	{
+		InHostBox->AddSlot()
+			.VAlign(VAlign_Center)
+			.AutoHeight()
+			.Padding(2.0f)
+			[
+				HierarchyBreadcrumbTrailWidget.ToSharedRef()
+			];
+	}
+	else if (ToolbarWidget)
+	{
+		InHostBox->AddSlot()
 			.VAlign(VAlign_Center)
+			.AutoHeight()
+			.Padding(2.0f)
 			[
 				ToolbarWidget.ToSharedRef()
-			]
-		];
+			];
 	}
-	else
+
+	TSharedPtr<SWidget> ExtraToolbarWidget = ConstructExtraToolbar();
+	if (ExtraToolbarWidget)
 	{
-		InWidgetContent->AddSlot()
-		.VAlign(VAlign_Center)
-		.AutoHeight()
-		.Padding(2.0f)
-		[
-			ConstructHierarchyBreadcrumbTrail()
-		];
+		InHostBox->AddSlot()
+			.VAlign(VAlign_Center)
+			.AutoHeight()
+			.Padding(2.0f)
+			[
+				ExtraToolbarWidget.ToSharedRef()
+			];
 	}
 }
 END_SLATE_FUNCTION_BUILD_OPTIMIZATION
@@ -444,17 +481,17 @@ END_SLATE_FUNCTION_BUILD_OPTIMIZATION
 ////////////////////////////////////////////////////////////////////////////////////////////////////
 
 BEGIN_SLATE_FUNCTION_BUILD_OPTIMIZATION
-void STableTreeView::ConstructFooterArea(TSharedRef<SVerticalBox> InWidgetContent)
+void STableTreeView::ConstructFooterArea(TSharedRef<SVerticalBox> InHostBox)
 {
 	TSharedPtr<SWidget> FooterWidget = ConstructFooter();
 	if (FooterWidget.IsValid())
 	{
-		InWidgetContent->AddSlot()
-		.AutoHeight()
-		.Padding(0.0f, 6.0f, 0.0f, 0.0f)
-		[
-			FooterWidget.ToSharedRef()
-		];
+		InHostBox->AddSlot()
+			.AutoHeight()
+			.Padding(0.0f, 6.0f, 0.0f, 0.0f)
+			[
+				FooterWidget.ToSharedRef()
+			];
 	}
 }
 END_SLATE_FUNCTION_BUILD_OPTIMIZATION
@@ -1040,7 +1077,9 @@ void STableTreeView::ApplyNodeFiltering()
 
 	// TableRowNodes --> FilteredNodes / FilteredNodesPtr
 	FilteredNodes.Reset();
-	if (FilterConfiguratorToUse && !FilterConfiguratorToUse->IsEmpty())
+	const bool bHasValidFilterConfigurator = FilterConfiguratorToUse && !FilterConfiguratorToUse->IsEmpty();
+	const bool bHasCustomNodeFilter = HasCustomNodeFilter();
+	if (bHasValidFilterConfigurator || bHasCustomNodeFilter)
 	{
 		for (const FTableTreeNodePtr& NodePtr : TableRowNodes)
 		{
@@ -1048,17 +1087,25 @@ void STableTreeView::ApplyNodeFiltering()
 			{
 				break;
 			}
-
-			if (FilterNode(*FilterConfiguratorToUse, *NodePtr))
+			if ((!bHasValidFilterConfigurator || FilterNode(*FilterConfiguratorToUse, *NodePtr)) &&
+				(!bHasCustomNodeFilter || FilterNodeCustom(*NodePtr)))
 			{
 				FilteredNodes.Add(NodePtr);
 			}
 		}
-		FilteredNodesPtr = &FilteredNodes;
+		if (FilteredNodes.Num() != TableRowNodes.Num())
+		{
+			FilteredNodesPtr = &FilteredNodes;
+		}
+		else
+		{
+			FilteredNodes.Reset();
+			FilteredNodesPtr = &TableRowNodes; // no filter; includes all nodes
+		}
 	}
 	else
 	{
-		FilteredNodesPtr = &TableRowNodes;
+		FilteredNodesPtr = &TableRowNodes; // no filter; includes all nodes
 	}
 
 	Stopwatch.Stop();
@@ -2021,16 +2068,14 @@ void STableTreeView::PostChangeGroupings()
 	for (TSharedPtr<FTreeNodeGrouping>& GroupingPtr : CurrentGroupings)
 	{
 		const FName& ColumnId = GroupingPtr->GetColumnId();
-
 		if (ColumnId != NAME_None)
 		{
 			if (bAdjustHierarchyColumnWidth)
 			{
 				// Compute width for Hierarchy column based on column used in grouping and its indentation.
-				const int32 NumColumns = TreeViewHeaderRow->GetColumns().Num();
-				for (int32 ColumnIndex = 0; ColumnIndex < NumColumns; ++ColumnIndex)
+				const TIndirectArray<SHeaderRow::FColumn>& HeaderRowColumns = TreeViewHeaderRow->GetColumns();
+				for (const SHeaderRow::FColumn& CurrentColumn : HeaderRowColumns)
 				{
-					const SHeaderRow::FColumn& CurrentColumn = TreeViewHeaderRow->GetColumns()[ColumnIndex];
 					if (CurrentColumn.ColumnId == ColumnId)
 					{
 						const float Width = HierarchyMinWidth + GroupingDepth + CurrentColumn.GetWidth();
@@ -3940,6 +3985,40 @@ FText STableTreeView::ViewPreset_GetSelectedToolTipText() const
 
 ////////////////////////////////////////////////////////////////////////////////////////////////////
 
+BEGIN_SLATE_FUNCTION_BUILD_OPTIMIZATION
+void STableTreeView::ConstructViewPreset(TSharedPtr<SHorizontalBox> Box)
+{
+	Box->AddSlot()
+		.AutoWidth()
+		.VAlign(VAlign_Center)
+		.Padding(0.0f, 0.0f, 4.0f, 0.0f)
+		[
+			SNew(STextBlock)
+			.Text(LOCTEXT("Preset", "Preset:"))
+		];
+
+	Box->AddSlot()
+		.AutoWidth()
+		.VAlign(VAlign_Center)
+		[
+			SNew(SBox)
+			[
+				SAssignNew(PresetComboBox, SComboBox<TSharedRef<ITableTreeViewPreset>>)
+				.ToolTipText(this, &STableTreeView::ViewPreset_GetSelectedToolTipText)
+				.OptionsSource(GetAvailableViewPresets())
+				.OnSelectionChanged(this, &STableTreeView::ViewPreset_OnSelectionChanged)
+				.OnGenerateWidget(this, &STableTreeView::ViewPreset_OnGenerateWidget)
+				[
+					SNew(STextBlock)
+					.Text(this, &STableTreeView::ViewPreset_GetSelectedText)
+				]
+			]
+		];
+}
+END_SLATE_FUNCTION_BUILD_OPTIMIZATION
+
+////////////////////////////////////////////////////////////////////////////////////////////////////
+
 void STableTreeView::StopAllTableDataTasks(bool bWait)
 {
 	DataTaskInfos.RemoveAllSwap([](TSharedPtr<FTableTaskInfo> Data) { return Data->Event->IsComplete(); });
diff --git a/Engine/Source/Developer/TraceInsightsCore/Public/InsightsCore/Table/Widgets/STableTreeView.h b/Engine/Source/Developer/TraceInsightsCore/Public/InsightsCore/Table/Widgets/STableTreeView.h
--- a/Engine/Source/Developer/TraceInsightsCore/Public/InsightsCore/Table/Widgets/STableTreeView.h
+++ b/Engine/Source/Developer/TraceInsightsCore/Public/InsightsCore/Table/Widgets/STableTreeView.h
@@ -188,11 +188,13 @@ protected:
 	void InitCommandList();
 
 	virtual void ConstructWidget(TSharedPtr<FTable> InTablePtr);
+	virtual TSharedPtr<SWidget> ConstructFilterToolbar();
 	virtual TSharedRef<SWidget> ConstructHierarchyBreadcrumbTrail();
 	virtual TSharedPtr<SWidget> ConstructToolbar() { return nullptr; }
+	virtual TSharedPtr<SWidget> ConstructExtraToolbar() { return nullptr; }
 	virtual TSharedPtr<SWidget> ConstructFooter() { return nullptr; }
-	virtual void ConstructHeaderArea(TSharedRef<SVerticalBox> InWidgetContent);
-	virtual void ConstructFooterArea(TSharedRef<SVerticalBox> InWidgetContent);
+	virtual void ConstructHeaderArea(TSharedRef<SVerticalBox> InHostBox);
+	virtual void ConstructFooterArea(TSharedRef<SVerticalBox> InHostBox);
 
 	void UpdateTree();
 
@@ -275,6 +277,9 @@ protected:
 
 	void ApplyNodeFiltering();
 
+	virtual bool HasCustomNodeFilter() const { return false; }
+	virtual bool FilterNodeCustom(const FTableTreeNode& InNode) const { return true; }
+
 	virtual bool FilterNode(const FFilterConfigurator& InFilterConfigurator, const FTableTreeNode& InNode) const;
 
 	virtual void InitFilterConfigurator(FFilterConfigurator& InOutFilterConfigurator);
@@ -469,6 +474,7 @@ protected:
 	TSharedRef<SWidget> ViewPreset_OnGenerateWidget(TSharedRef<ITableTreeViewPreset> InPreset);
 	FText ViewPreset_GetSelectedText() const;
 	FText ViewPreset_GetSelectedToolTipText() const;
+	void ConstructViewPreset(TSharedPtr<SHorizontalBox> InHorizontalBox);
 
 	////////////////////////////////////////////////////////////////////////////////////////////////////
 
