[Insights] Table/Tree views: Fixed tooltip not updating correctly for the columns with "IsDynamic" flag.

diff --git a/Engine/Source/Developer/TraceInsightsCore/Private/InsightsCore/Table/ViewModels/TableCellValueFormatter.cpp b/Engine/Source/Developer/TraceInsightsCore/Private/InsightsCore/Table/ViewModels/TableCellValueFormatter.cpp
--- a/Engine/Source/Developer/TraceInsightsCore/Private/InsightsCore/Table/ViewModels/TableCellValueFormatter.cpp
+++ b/Engine/Source/Developer/TraceInsightsCore/Private/InsightsCore/Table/ViewModels/TableCellValueFormatter.cpp
@@ -57,18 +57,37 @@ FText FTableCellValueFormatter::CopyTooltip(const FTableColumn& Column, const FB
 
 TSharedPtr<IToolTip> FTableCellValueFormatter::GetCustomTooltip(const FTableColumn& Column, const FBaseTreeNode& Node) const
 {
-	return SNew(SToolTip)
-		.Visibility(TAttribute<EVisibility>::Create(TAttribute<EVisibility>::FGetter::CreateStatic(&FTableCellValueFormatter::GetTooltipVisibility)))
-		[
-			SNew(SVerticalBox)
-			+ SVerticalBox::Slot()
-			.AutoHeight()
-			.Padding(2.0f)
+	if (Column.IsDynamic())
+	{
+		return SNew(SToolTip)
+			.Visibility(TAttribute<EVisibility>::Create(TAttribute<EVisibility>::FGetter::CreateStatic(&FTableCellValueFormatter::GetTooltipVisibility)))
 			[
-				SNew(STextBlock)
-				.Text(FormatValueForTooltip(Column.GetValueGetter()->GetValue(Column, Node)))
-			]
-		];
+				SNew(SVerticalBox)
+				+ SVerticalBox::Slot()
+				.AutoHeight()
+				.Padding(2.0f)
+				[
+					SNew(STextBlock)
+					.Text_Lambda([this, &Column, &Node]
+					{ return FormatValueForTooltip(Column.GetValueGetter()->GetValue(Column, Node)); })
+				]
+			];
+	}
+	else
+	{
+		return SNew(SToolTip)
+			.Visibility(TAttribute<EVisibility>::Create(TAttribute<EVisibility>::FGetter::CreateStatic(&FTableCellValueFormatter::GetTooltipVisibility)))
+			[
+				SNew(SVerticalBox)
+				+ SVerticalBox::Slot()
+				.AutoHeight()
+				.Padding(2.0f)
+				[
+					SNew(STextBlock)
+					.Text(FormatValueForTooltip(Column.GetValueGetter()->GetValue(Column, Node)))
+				]
+			];
+	}
 }
 
 ////////////////////////////////////////////////////////////////////////////////////////////////////
