[Insights] TraceInsightsCore: Added FormattingOptions to FInt64ValueFormatterAsMemory. Also added public "static FText FormatForTooltip(int64 InValue);" to FInt64ValueFormatterAsMemory class.

diff --git a/Engine/Source/Developer/TraceInsightsCore/Private/InsightsCore/Table/ViewModels/TableCellValueFormatter.cpp b/Engine/Source/Developer/TraceInsightsCore/Private/InsightsCore/Table/ViewModels/TableCellValueFormatter.cpp
--- a/Engine/Source/Developer/TraceInsightsCore/Private/InsightsCore/Table/ViewModels/TableCellValueFormatter.cpp
+++ b/Engine/Source/Developer/TraceInsightsCore/Private/InsightsCore/Table/ViewModels/TableCellValueFormatter.cpp
@@ -160,8 +160,6 @@ FText FInt64ValueFormatterAsMemory::FormatValue(const TOptional<FTableCellValue>
 		const int64 Value = InValue.GetValue().Int64;
 		if (Value > 0)
 		{
-			FNumberFormattingOptions FormattingOptions;
-			FormattingOptions.MaximumFractionalDigits = 1;
 			return FText::AsMemory(Value, &FormattingOptions);
 		}
 		else if (Value == 0)
@@ -170,8 +168,6 @@ FText FInt64ValueFormatterAsMemory::FormatValue(const TOptional<FTableCellValue>
 		}
 		else // Value < 0
 		{
-			FNumberFormattingOptions FormattingOptions;
-			FormattingOptions.MaximumFractionalDigits = 1;
 			return FText::Format(LOCTEXT("AsMemory_NegativeValue_Fmt1", "-{0}"), FText::AsMemory(-Value, &FormattingOptions));
 		}
 	}
@@ -182,55 +178,59 @@ FText FInt64ValueFormatterAsMemory::FormatValue(const TOptional<FTableCellValue>
 
 FText FInt64ValueFormatterAsMemory::FormatValueForTooltip(const TOptional<FTableCellValue>& InValue) const
 {
-	if (InValue.IsSet())
+	return InValue.IsSet() ? FormatForTooltip(InValue.GetValue().Int64) : FText::GetEmpty();
+}
+
+////////////////////////////////////////////////////////////////////////////////////////////////////
+
+FText FInt64ValueFormatterAsMemory::FormatForTooltip(int64 Value)
+{
+	if (Value > 0)
 	{
-		const int64 Value = InValue.GetValue().Int64;
-		if (Value > 0)
+		if (Value < 1024)
 		{
-			if (Value < 1024)
+			if (Value == 1)
 			{
-				if (Value == 1)
-				{
-					return LOCTEXT("AsMemory_Tooltip_1byte", "1 byte");
-				}
-				else
-				{
-					return FText::Format(LOCTEXT("AsMemory_PositiveValue_TooltipFmt1", "{0} bytes"), FText::AsNumber(Value));
-				}
+				return LOCTEXT("AsMemory_Tooltip_1byte", "1 byte");
 			}
 			else
 			{
-				FNumberFormattingOptions FormattingOptions;
-				FormattingOptions.MaximumFractionalDigits = 2;
-				return FText::Format(LOCTEXT("AsMemory_PositiveValue_TooltipFmt2", "{0} bytes ({1})"), FText::AsNumber(Value), FText::AsMemory(Value, &FormattingOptions));
+				return FText::Format(LOCTEXT("AsMemory_PositiveValue_TooltipFmt1", "{0} bytes"), FText::AsNumber(Value));
 			}
 		}
-		else if (Value == 0)
+		else
 		{
-			return LOCTEXT("AsMemory_ZeroValue", "0");
+			FNumberFormattingOptions FormattingOptions;
+			FormattingOptions.MaximumFractionalDigits = 2;
+			return FText::Format(LOCTEXT("AsMemory_PositiveValue_TooltipFmt2", "{0} bytes ({1})"), FText::AsNumber(Value), FText::AsMemory(Value, &FormattingOptions));
 		}
-		else // Value < 0
+	}
+
+	if (Value == 0)
+	{
+		return LOCTEXT("AsMemory_ZeroValue", "0");
+	}
+
+	// Value < 0
+	{
+		if (-Value < 1024)
 		{
-			if (-Value < 1024)
+			if (Value == -1)
 			{
-				if (Value == -1)
-				{
-					return LOCTEXT("AsMemory_Tooltip_minus1byte", "-1 byte");
-				}
-				else
-				{
-					return FText::Format(LOCTEXT("AsMemory_NegativeValue_TooltipFmt1", "-{0} bytes"), FText::AsNumber(-Value));
-				}
+				return LOCTEXT("AsMemory_Tooltip_minus1byte", "-1 byte");
 			}
 			else
 			{
-				FNumberFormattingOptions FormattingOptions;
-				FormattingOptions.MaximumFractionalDigits = 2;
-				return FText::Format(LOCTEXT("AsMemory_NegativeValue_TooltipFmt2", "-{0} bytes (-{1})"), FText::AsNumber(-Value), FText::AsMemory(-Value, &FormattingOptions));
+				return FText::Format(LOCTEXT("AsMemory_NegativeValue_TooltipFmt1", "-{0} bytes"), FText::AsNumber(-Value));
 			}
 		}
+		else
+		{
+			FNumberFormattingOptions FormattingOptions;
+			FormattingOptions.MaximumFractionalDigits = 2;
+			return FText::Format(LOCTEXT("AsMemory_NegativeValue_TooltipFmt2", "-{0} bytes (-{1})"), FText::AsNumber(-Value), FText::AsMemory(-Value, &FormattingOptions));
+		}
 	}
-	return FText::GetEmpty();
 }
 
 ////////////////////////////////////////////////////////////////////////////////////////////////////
diff --git a/Engine/Source/Developer/TraceInsightsCore/Public/InsightsCore/Table/ViewModels/TableCellValueFormatter.h b/Engine/Source/Developer/TraceInsightsCore/Public/InsightsCore/Table/ViewModels/TableCellValueFormatter.h
--- a/Engine/Source/Developer/TraceInsightsCore/Public/InsightsCore/Table/ViewModels/TableCellValueFormatter.h
+++ b/Engine/Source/Developer/TraceInsightsCore/Public/InsightsCore/Table/ViewModels/TableCellValueFormatter.h
@@ -150,8 +150,22 @@ public:
 class TRACEINSIGHTSCORE_API FInt64ValueFormatterAsMemory : public FTableCellValueFormatter
 {
 public:
+	FInt64ValueFormatterAsMemory()
+	{
+		FormattingOptions.MaximumFractionalDigits = 1;
+	}
+	virtual ~FInt64ValueFormatterAsMemory() {}
+
+	const FNumberFormattingOptions& GetFormattingOptions() const { return FormattingOptions; }
+	FNumberFormattingOptions& GetFormattingOptions() { return FormattingOptions; }
+
 	virtual FText FormatValue(const TOptional<FTableCellValue>& InValue) const override;
 	virtual FText FormatValueForTooltip(const TOptional<FTableCellValue>& InValue) const override;
+
+	static FText FormatForTooltip(int64 InValue);
+
+private:
+	FNumberFormattingOptions FormattingOptions;
 };
 
 ////////////////////////////////////////////////////////////////////////////////////////////////////
