[Insights] Deprecated GetBrushName() for ITreeNodeGrouping.

diff --git a/Engine/Plugins/Editor/AssetManagerEditor/Source/AssetManagerEditor/Private/TreeView/AssetDependencyGrouping.cpp b/Engine/Plugins/Editor/AssetManagerEditor/Source/AssetManagerEditor/Private/TreeView/AssetDependencyGrouping.cpp
--- a/Engine/Plugins/Editor/AssetManagerEditor/Source/AssetManagerEditor/Private/TreeView/AssetDependencyGrouping.cpp
+++ b/Engine/Plugins/Editor/AssetManagerEditor/Source/AssetManagerEditor/Private/TreeView/AssetDependencyGrouping.cpp
@@ -20,7 +20,6 @@ FAssetDependencyGrouping::FAssetDependencyGrouping()
 		LOCTEXT("Grouping_ByAssetDependency_ShortName", "Asset Dependency"),
 		LOCTEXT("Grouping_ByAssetDependency_TitleName", "By Asset Dependency"),
 		LOCTEXT("Grouping_ByAssetDependency_Desc", "Group assets based on their dependency."),
-		TEXT("Icons.Group.TreeItem"),
 		nullptr)
 {
 }
@@ -127,7 +126,6 @@ FPluginDependencyGrouping::FPluginDependencyGrouping()
 		LOCTEXT("Grouping_ByPluginDependency_ShortName", "Plugin Dependency"),
 		LOCTEXT("Grouping_ByPluginDependency_TitleName", "By Plugin Dependency"),
 		LOCTEXT("Grouping_ByPluginDependency_Desc", "Group assets based on their plugin. Each plugin group node will all show the dependent plugins."),
-		TEXT("Icons.Group.TreeItem"),
 		nullptr)
 {
 }
diff --git a/Engine/Source/Developer/TraceInsights/Private/Insights/MemoryProfiler/ViewModels/MemAllocGroupingByCallstack.cpp b/Engine/Source/Developer/TraceInsights/Private/Insights/MemoryProfiler/ViewModels/MemAllocGroupingByCallstack.cpp
--- a/Engine/Source/Developer/TraceInsights/Private/Insights/MemoryProfiler/ViewModels/MemAllocGroupingByCallstack.cpp
+++ b/Engine/Source/Developer/TraceInsights/Private/Insights/MemoryProfiler/ViewModels/MemAllocGroupingByCallstack.cpp
@@ -79,7 +79,6 @@ FMemAllocGroupingByCallstack::FMemAllocGroupingByCallstack(bool bInIsAllocCallst
 		: (bInIsInverted ? LOCTEXT("Grouping_ByCallstack4_TitleName", "By Inverted Free Callstack")
 						 : LOCTEXT("Grouping_ByCallstack3_TitleName", "By Free Callstack")),
 		LOCTEXT("Grouping_Callstack_Desc", "Creates a tree based on callstack of each allocation."),
-		TEXT("Icons.Group.TreeItem"),
 		nullptr)
 	, bIsAllocCallstack(bInIsAllocCallstack)
 	, bIsInverted(bInIsInverted)
diff --git a/Engine/Source/Developer/TraceInsights/Private/Insights/MemoryProfiler/ViewModels/MemAllocGroupingByHeap.cpp b/Engine/Source/Developer/TraceInsights/Private/Insights/MemoryProfiler/ViewModels/MemAllocGroupingByHeap.cpp
--- a/Engine/Source/Developer/TraceInsights/Private/Insights/MemoryProfiler/ViewModels/MemAllocGroupingByHeap.cpp
+++ b/Engine/Source/Developer/TraceInsights/Private/Insights/MemoryProfiler/ViewModels/MemAllocGroupingByHeap.cpp
@@ -28,7 +28,6 @@ FMemAllocGroupingByHeap::FMemAllocGroupingByHeap(const TraceServices::IAllocatio
 		LOCTEXT("Grouping_ByHeap_ShortName", "Heap"),
 		LOCTEXT("Grouping_ByHeap_TitleName", "By Heap"),
 		LOCTEXT("Grouping_ByHeap_Desc", "Creates a tree based on heap."),
-		TEXT("Icons.Group.TreeItem"),
 		nullptr)
 	, AllocProvider(InAllocProvider)
 {
diff --git a/Engine/Source/Developer/TraceInsights/Private/Insights/MemoryProfiler/ViewModels/MemAllocGroupingBySize.cpp b/Engine/Source/Developer/TraceInsights/Private/Insights/MemoryProfiler/ViewModels/MemAllocGroupingBySize.cpp
--- a/Engine/Source/Developer/TraceInsights/Private/Insights/MemoryProfiler/ViewModels/MemAllocGroupingBySize.cpp
+++ b/Engine/Source/Developer/TraceInsights/Private/Insights/MemoryProfiler/ViewModels/MemAllocGroupingBySize.cpp
@@ -26,7 +26,6 @@ FMemAllocGroupingBySize::FMemAllocGroupingBySize()
 		LOCTEXT("Grouping_BySize_ShortName", "Size"),
 		LOCTEXT("Grouping_BySize_TitleName", "By Size"),
 		LOCTEXT("Grouping_BySize_Desc", "Group allocations based on their size."),
-		TEXT("Icons.Group.TreeItem"),
 		nullptr)
 {
 	ResetThresholdsPow2();
diff --git a/Engine/Source/Developer/TraceInsights/Private/Insights/MemoryProfiler/ViewModels/MemAllocGroupingBySwapPage.cpp b/Engine/Source/Developer/TraceInsights/Private/Insights/MemoryProfiler/ViewModels/MemAllocGroupingBySwapPage.cpp
--- a/Engine/Source/Developer/TraceInsights/Private/Insights/MemoryProfiler/ViewModels/MemAllocGroupingBySwapPage.cpp
+++ b/Engine/Source/Developer/TraceInsights/Private/Insights/MemoryProfiler/ViewModels/MemAllocGroupingBySwapPage.cpp
@@ -28,7 +28,6 @@ FMemAllocGroupingBySwapPage::FMemAllocGroupingBySwapPage(const TraceServices::IA
 		LOCTEXT("Grouping_BySwap_ShortName", "Swap"),
 		LOCTEXT("Grouping_BySwap_TitleName", "By Swap"),
 		LOCTEXT("Grouping_BySwap_Desc", "Creates a tree based on swap state."),
-		TEXT("Icons.Group.TreeItem"),
 		nullptr)
 	, AllocProvider(InAllocProvider)
 {
diff --git a/Engine/Source/Developer/TraceInsights/Private/Insights/MemoryProfiler/ViewModels/MemAllocGroupingByTag.cpp b/Engine/Source/Developer/TraceInsights/Private/Insights/MemoryProfiler/ViewModels/MemAllocGroupingByTag.cpp
--- a/Engine/Source/Developer/TraceInsights/Private/Insights/MemoryProfiler/ViewModels/MemAllocGroupingByTag.cpp
+++ b/Engine/Source/Developer/TraceInsights/Private/Insights/MemoryProfiler/ViewModels/MemAllocGroupingByTag.cpp
@@ -34,7 +34,6 @@ FMemAllocGroupingByTag::FMemAllocGroupingByTag(const TraceServices::IAllocations
 		LOCTEXT("Grouping_ByTag_ShortName", "Tag"),
 		LOCTEXT("Grouping_ByTag_TitleName", "By Tag"),
 		LOCTEXT("Grouping_Tag_Desc", "Creates a tree based on Tag hierarchy."),
-		TEXT("Icons.Group.TreeItem"),
 		nullptr)
 	, TagProvider(InTagProvider)
 {
diff --git a/Engine/Source/Developer/TraceInsightsCore/Private/InsightsCore/Table/ViewModels/TreeNodeGrouping.cpp b/Engine/Source/Developer/TraceInsightsCore/Private/InsightsCore/Table/ViewModels/TreeNodeGrouping.cpp
--- a/Engine/Source/Developer/TraceInsightsCore/Private/InsightsCore/Table/ViewModels/TreeNodeGrouping.cpp
+++ b/Engine/Source/Developer/TraceInsightsCore/Private/InsightsCore/Table/ViewModels/TreeNodeGrouping.cpp
@@ -22,6 +22,17 @@ INSIGHTS_IMPLEMENT_RTTI(FTreeNodeGroupingByPathBreakdown)
 // FTreeNodeGrouping
 ////////////////////////////////////////////////////////////////////////////////////////////////////
 
+FTreeNodeGrouping::FTreeNodeGrouping(const FText& InShortName, const FText& InTitleName, const FText& InDescription, const FSlateBrush* InIcon)
+	: ShortName(InShortName)
+	, TitleName(InTitleName)
+	, Description(InDescription)
+	, BrushName(NAME_None)
+	, Icon(InIcon)
+{
+}
+
+////////////////////////////////////////////////////////////////////////////////////////////////////
+
 FTreeNodeGrouping::FTreeNodeGrouping(const FText& InShortName, const FText& InTitleName, const FText& InDescription, const FName InBrushName, const FSlateBrush* InIcon)
 	: ShortName(InShortName)
 	, TitleName(InTitleName)
@@ -33,6 +44,15 @@ FTreeNodeGrouping::FTreeNodeGrouping(const FText& InShortName, const FText& InTi
 
 ////////////////////////////////////////////////////////////////////////////////////////////////////
 
+FName FTreeNodeGrouping::GetBrushName() const
+{
+	PRAGMA_DISABLE_DEPRECATION_WARNINGS
+	return BrushName;
+	PRAGMA_ENABLE_DEPRECATION_WARNINGS
+}
+
+////////////////////////////////////////////////////////////////////////////////////////////////////
+
 void FTreeNodeGrouping::GroupNodes(const TArray<FTableTreeNodePtr>& Nodes, FTableTreeNode& ParentGroup, TWeakPtr<FTable> InParentTable, IAsyncOperationProgress& InAsyncOperationProgress) const
 {
 	TMap<FName, FTableTreeNodePtr> GroupMap;
@@ -81,7 +101,6 @@ FTreeNodeGroupingFlat::FTreeNodeGroupingFlat()
 		LOCTEXT("Grouping_Flat_ShortName", "All"),
 		LOCTEXT("Grouping_Flat_TitleName", "Flat (All)"),
 		LOCTEXT("Grouping_Flat_Desc", "Creates a single group. Includes all items."),
-		TEXT("Icons.Group.TreeItem"),
 		nullptr)
 {
 }
@@ -117,7 +136,6 @@ FTreeNodeGroupingByUniqueValue::FTreeNodeGroupingByUniqueValue(TSharedRef<FTable
 		FText::Format(LOCTEXT("Grouping_ByUniqueValue_ShortNameFmt", "{0} (Unique)"), InColumnRef->GetTitleName()),
 		FText::Format(LOCTEXT("Grouping_ByUniqueValue_TitleNameFmt", "Unique Values - {0}"), InColumnRef->GetTitleName()),
 		LOCTEXT("Grouping_ByUniqueValue_Desc", "Creates a group for each unique value."),
-		TEXT("Icons.Group.TreeItem"),
 		nullptr)
 	, ColumnRef(InColumnRef)
 {
@@ -251,7 +269,6 @@ FTreeNodeGroupingByNameFirstLetter::FTreeNodeGroupingByNameFirstLetter()
 		LOCTEXT("Grouping_ByName_ShortName", "Name"),
 		LOCTEXT("Grouping_ByName_TitleName", "Name (First Letter)"),
 		LOCTEXT("Grouping_ByName_Desc", "Creates a group for each first letter of node names."),
-		TEXT("Icons.Group.TreeItem"),
 		nullptr)
 {
 }
@@ -272,7 +289,6 @@ FTreeNodeGroupingByType::FTreeNodeGroupingByType()
 		LOCTEXT("Grouping_ByTypeName_ShortName", "TypeName"),
 		LOCTEXT("Grouping_ByTypeName_TitleName", "Type Name"),
 		LOCTEXT("Grouping_ByTypeName_Desc", "Creates a group for each node type."),
-		TEXT("Icons.Group.TreeItem"),
 		nullptr)
 {
 }
@@ -293,7 +309,6 @@ FTreeNodeGroupingByPathBreakdown::FTreeNodeGroupingByPathBreakdown(TSharedRef<FT
 		FText::Format(LOCTEXT("Grouping_ByPathBreakdown_ShortNameFmt", "{0} (Path)"), InColumnRef->GetTitleName()),
 		FText::Format(LOCTEXT("Grouping_ByPathBreakdown_TitleNameFmt", "Path Breakdown - {0}"), InColumnRef->GetTitleName()),
 		LOCTEXT("Grouping_ByPathBreakdown_Desc", "Creates a tree hierarchy out of the path structure of string values."),
-		TEXT("Icons.Group.TreeItem"),
 		nullptr)
 	, ColumnRef(InColumnRef)
 {
diff --git a/Engine/Source/Developer/TraceInsightsCore/Public/InsightsCore/Table/ViewModels/TreeNodeGrouping.h b/Engine/Source/Developer/TraceInsightsCore/Public/InsightsCore/Table/ViewModels/TreeNodeGrouping.h
--- a/Engine/Source/Developer/TraceInsightsCore/Public/InsightsCore/Table/ViewModels/TreeNodeGrouping.h
+++ b/Engine/Source/Developer/TraceInsightsCore/Public/InsightsCore/Table/ViewModels/TreeNodeGrouping.h
@@ -38,7 +38,9 @@ public:
 	virtual FText GetTitleName() const = 0;
 	virtual FText GetDescription() const = 0;
 
-	virtual FName GetBrushName() const = 0;
+	UE_DEPRECATED(5.6, "GetBrushName() is not used")
+	virtual FName GetBrushName() const { return NAME_None; }
+
 	virtual const FSlateBrush* GetIcon() const = 0;
 
 	virtual FName GetColumnId() const = 0;
@@ -59,6 +61,8 @@ class TRACEINSIGHTSCORE_API FTreeNodeGrouping : public ITreeNodeGrouping
 	INSIGHTS_DECLARE_RTTI(FTreeNodeGrouping, ITreeNodeGrouping)
 
 public:
+	FTreeNodeGrouping(const FText& InShortName, const FText& InTitleName, const FText& InDescription, const FSlateBrush* InIcon);
+	UE_DEPRECATED(5.6, "BrushName is not used")
 	FTreeNodeGrouping(const FText& InShortName, const FText& InTitleName, const FText& InDescription, const FName InBrushName, const FSlateBrush* InIcon);
 	virtual ~FTreeNodeGrouping() {}
 
@@ -66,7 +70,9 @@ public:
 	virtual FText GetTitleName() const override { return TitleName; }
 	virtual FText GetDescription() const override { return Description; }
 
-	virtual FName GetBrushName() const override { return BrushName; }
+	UE_DEPRECATED(5.6, "GetBrushName() is not used")
+	virtual FName GetBrushName() const override;
+
 	virtual const FSlateBrush* GetIcon() const override { return Icon; }
 
 	virtual FName GetColumnId() const override { return NAME_None; }
@@ -78,6 +84,7 @@ protected:
 	FText ShortName;
 	FText TitleName;
 	FText Description;
+	UE_DEPRECATED(5.6, "BrushName is not used")
 	FName BrushName;
 	const FSlateBrush* Icon;
 };
