Fixed deprecation warning when compiled with Clang.

#rnx

diff --git a/Engine/Source/Developer/TraceInsightsCore/Private/InsightsCore/Table/ViewModels/TreeNodeGrouping.cpp b/Engine/Source/Developer/TraceInsightsCore/Private/InsightsCore/Table/ViewModels/TreeNodeGrouping.cpp
--- a/Engine/Source/Developer/TraceInsightsCore/Private/InsightsCore/Table/ViewModels/TreeNodeGrouping.cpp
+++ b/Engine/Source/Developer/TraceInsightsCore/Private/InsightsCore/Table/ViewModels/TreeNodeGrouping.cpp
@@ -22,6 +22,7 @@ INSIGHTS_IMPLEMENT_RTTI(FTreeNodeGroupingByPathBreakdown)
 // FTreeNodeGrouping
 ////////////////////////////////////////////////////////////////////////////////////////////////////
 
+PRAGMA_DISABLE_DEPRECATION_WARNINGS
 FTreeNodeGrouping::FTreeNodeGrouping(const FText& InShortName, const FText& InTitleName, const FText& InDescription, const FSlateBrush* InIcon)
 	: ShortName(InShortName)
 	, TitleName(InTitleName)
@@ -30,9 +31,11 @@ FTreeNodeGrouping::FTreeNodeGrouping(const FText& InShortName, const FText& InTi
 	, Icon(InIcon)
 {
 }
+PRAGMA_ENABLE_DEPRECATION_WARNINGS
 
 ////////////////////////////////////////////////////////////////////////////////////////////////////
 
+PRAGMA_DISABLE_DEPRECATION_WARNINGS
 FTreeNodeGrouping::FTreeNodeGrouping(const FText& InShortName, const FText& InTitleName, const FText& InDescription, const FName InBrushName, const FSlateBrush* InIcon)
 	: ShortName(InShortName)
 	, TitleName(InTitleName)
@@ -41,6 +44,7 @@ FTreeNodeGrouping::FTreeNodeGrouping(const FText& InShortName, const FText& InTi
 	, Icon(InIcon)
 {
 }
+PRAGMA_ENABLE_DEPRECATION_WARNINGS
 
 ////////////////////////////////////////////////////////////////////////////////////////////////////
 
diff --git a/Engine/Source/Developer/TraceInsightsCore/Public/InsightsCore/Table/ViewModels/TreeNodeGrouping.h b/Engine/Source/Developer/TraceInsightsCore/Public/InsightsCore/Table/ViewModels/TreeNodeGrouping.h
--- a/Engine/Source/Developer/TraceInsightsCore/Public/InsightsCore/Table/ViewModels/TreeNodeGrouping.h
+++ b/Engine/Source/Developer/TraceInsightsCore/Public/InsightsCore/Table/ViewModels/TreeNodeGrouping.h
@@ -61,9 +61,15 @@ class TRACEINSIGHTSCORE_API FTreeNodeGrouping : public ITreeNodeGrouping
 	INSIGHTS_DECLARE_RTTI(FTreeNodeGrouping, ITreeNodeGrouping)
 
 public:
+	FTreeNodeGrouping() : Icon(nullptr) {}
 	FTreeNodeGrouping(const FText& InShortName, const FText& InTitleName, const FText& InDescription, const FSlateBrush* InIcon);
+
 	UE_DEPRECATED(5.6, "BrushName is not used")
 	FTreeNodeGrouping(const FText& InShortName, const FText& InTitleName, const FText& InDescription, const FName InBrushName, const FSlateBrush* InIcon);
+
+	FTreeNodeGrouping(const FTreeNodeGrouping&) = delete;
+	FTreeNodeGrouping& operator=(FTreeNodeGrouping&) = delete;
+
 	virtual ~FTreeNodeGrouping() {}
 
 	virtual FText GetShortName() const override { return ShortName; }
