[Insights] Tree Views: Updated style for the grouping breadcrumb control to look like a combo box.

diff --git a/Engine/Source/Developer/TraceInsightsCore/Private/InsightsCore/Table/Widgets/STableTreeView.cpp b/Engine/Source/Developer/TraceInsightsCore/Private/InsightsCore/Table/Widgets/STableTreeView.cpp
--- a/Engine/Source/Developer/TraceInsightsCore/Private/InsightsCore/Table/Widgets/STableTreeView.cpp
+++ b/Engine/Source/Developer/TraceInsightsCore/Private/InsightsCore/Table/Widgets/STableTreeView.cpp
@@ -12,6 +12,7 @@
 #include "Logging/MessageLog.h"
 #include "SlateOptMacros.h"
 #include "Styling/StyleColors.h"
+#include "Widgets/Input/SComboButton.h"
 #include "Widgets/Input/SSearchBox.h"
 #include "Widgets/Images/SImage.h"
 #include "Widgets/SToolTip.h"
@@ -348,29 +349,27 @@ END_SLATE_FUNCTION_BUILD_OPTIMIZATION
 BEGIN_SLATE_FUNCTION_BUILD_OPTIMIZATION
 TSharedRef<SWidget> STableTreeView::ConstructHierarchyBreadcrumbTrail()
 {
-	return SNew(SHorizontalBox)
+	TSharedRef<SWidget> Widget = SNew(SComboButton)
+		.HasDownArrow(false)
+		.ToolTipText(LOCTEXT("GroupingsTooltip", "Groupings Hierarchy"))
+		.ButtonContent()
+		[
+			SNew(SBox)
+			.HAlign(HAlign_Left)
+			[
+				SAssignNew(GroupingBreadcrumbTrail, SBreadcrumbTrail<TSharedPtr<FTreeNodeGrouping>>)
+				.ButtonStyle(FAppStyle::Get(), "SimpleButton")
+				.TextStyle(FAppStyle::Get(), "NormalText")
+				.DelimiterImage(FAppStyle::Get().GetBrush("Icons.ChevronRight"))
+				.ButtonContentPadding(FMargin(3.0f, 2.0f))
+				.GetCrumbMenuContent(this, &STableTreeView::GetGroupingCrumbMenuContent)
+				.OnCrumbClicked(this, &STableTreeView::OnGroupingCrumbClicked)
+				.ToolTipText(LOCTEXT("GroupingCrumbTooltip", "Click a grouping to remove the additional groupings.\nClick a separator to add/insert a new grouping or to change the current groupings."))
+			]
+		];
 
-	+ SHorizontalBox::Slot()
-	.AutoWidth()
-	.VAlign(VAlign_Center)
-	[
-		SNew(STextBlock)
-		.Text(LOCTEXT("GroupByText", "Hierarchy:"))
-		.Margin(FMargin(0.0f, 0.0f, 4.0f, 0.0f))
-	]
-
-	+ SHorizontalBox::Slot()
-	.FillWidth(1.0f)
-	.VAlign(VAlign_Center)
-	[
-		SAssignNew(GroupingBreadcrumbTrail, SBreadcrumbTrail<TSharedPtr<FTreeNodeGrouping>>)
-		.ButtonStyle(FAppStyle::Get(), "SimpleButton")
-		.ButtonContentPadding(FMargin(3.0f, 2.0f))
-		.TextStyle(FAppStyle::Get(), "NormalText")
-		.DelimiterImage(FAppStyle::Get().GetBrush("Icons.ChevronRight"))
-		.OnCrumbClicked(this, &STableTreeView::OnGroupingCrumbClicked)
-		.GetCrumbMenuContent(this, &STableTreeView::GetGroupingCrumbMenuContent)
-	];
+	GroupingBreadcrumbTrail->ScrollToEnd();
+	return Widget;
 }
 END_SLATE_FUNCTION_BUILD_OPTIMIZATION
 
@@ -2154,6 +2153,8 @@ void STableTreeView::RebuildGroupingCrumbs()
 	{
 		GroupingBreadcrumbTrail->PushCrumb(Grouping->GetShortName(), Grouping);
 	}
+
+	GroupingBreadcrumbTrail->ScrollToEnd();
 }
 
 ////////////////////////////////////////////////////////////////////////////////////////////////////
