[Insights] TraceInsightsCore: Fixed style for SearchBox to inherit the app style (ex. the rounded corners), but to keep the fix regarding centered vertical alignment of text. Also removed the (unneeded) customization of "ToggleButton" style.

#rnx

diff --git a/Engine/Source/Developer/TraceInsightsCore/Private/InsightsCore/Common/InsightsCoreStyle.cpp b/Engine/Source/Developer/TraceInsightsCore/Private/InsightsCore/Common/InsightsCoreStyle.cpp
--- a/Engine/Source/Developer/TraceInsightsCore/Private/InsightsCore/Common/InsightsCoreStyle.cpp
+++ b/Engine/Source/Developer/TraceInsightsCore/Private/InsightsCore/Common/InsightsCoreStyle.cpp
@@ -241,22 +241,17 @@ void FInsightsCoreStyle::FStyle::Initialize()
 
 	//////////////////////////////////////////////////
 
-	// NormalEditableTextBox && SearchBox
+	// NormalEditableTextBox, SearchBox
+	// Fixes the vertical alignment of text (inside editable text boxes) to be centered.
 	{
 		const FEditableTextBoxStyle& NormalEditableTextBoxStyle = FAppStyle::GetWidgetStyle<FEditableTextBoxStyle>("NormalEditableTextBox");
-		const FTextBlockStyle TextBlockStyle = FTextBlockStyle()
-			.SetColorAndOpacity(NormalEditableTextBoxStyle.ForegroundColor)
-			.SetHighlightColor(NormalEditableTextBoxStyle.FocusedForegroundColor)
-			.SetFont(NormalEditableTextBoxStyle.TextStyle.Font)
-			.SetFontSize(uint16(NormalEditableTextBoxStyle.TextStyle.Font.Size));
 		const FEditableTextBoxStyle EditableTextBoxStyle = FEditableTextBoxStyle(NormalEditableTextBoxStyle)
-			.SetPadding(FMargin(6.0f, 4.0f, 6.0f, 4.0f))
-			.SetTextStyle(TextBlockStyle);
+			.SetPadding(FMargin(6.0f, 4.0f, 6.0f, 4.0f));
 		Set("NormalEditableTextBox", EditableTextBoxStyle);
 
 		const FSearchBoxStyle& NormalSearchBoxStyle = FAppStyle::GetWidgetStyle<FSearchBoxStyle>("SearchBox");
-		const FSearchBoxStyle SearchBoxStyle = FSearchBoxStyle(NormalSearchBoxStyle)
-			.SetTextBoxStyle(EditableTextBoxStyle);
+		FSearchBoxStyle SearchBoxStyle = FSearchBoxStyle(NormalSearchBoxStyle);
+		SearchBoxStyle.TextBoxStyle.SetPadding(FMargin(6.0f, 4.0f, 6.0f, 4.0f));
 		Set("SearchBox", SearchBoxStyle);
 	}
 
@@ -322,14 +317,6 @@ void FInsightsCoreStyle::FStyle::Initialize()
 		Set("SecondaryToolbar2.MaxUniformToolbarSize", 32.0f);
 	}
 
-	// ToggleButton
-	{
-		Set("ToggleButton", FButtonStyle(Button)
-			.SetNormal(FSlateNoResource())
-			.SetHovered(EDITOR_BOX_BRUSH("Common/RoundedSelection_16x", 4.0f / 16.0f, SelectionColor))
-			.SetPressed(EDITOR_BOX_BRUSH("Common/RoundedSelection_16x", 4.0f / 16.0f, SelectionColor_Pressed)));
-	}
-
 	// Common.GotoNativeCodeHyperlink
 	{
 		FTextBlockStyle InheritedFromNativeTextStyle = FTextBlockStyle(NormalText)
