[Insights] Tree Table Views: Removed Search combobox from Export context menu.

#rnx

diff --git a/Engine/Source/Developer/TraceInsightsCore/Private/InsightsCore/Table/Widgets/STableTreeView.cpp b/Engine/Source/Developer/TraceInsightsCore/Private/InsightsCore/Table/Widgets/STableTreeView.cpp
--- a/Engine/Source/Developer/TraceInsightsCore/Private/InsightsCore/Table/Widgets/STableTreeView.cpp
+++ b/Engine/Source/Developer/TraceInsightsCore/Private/InsightsCore/Table/Widgets/STableTreeView.cpp
@@ -743,7 +743,9 @@ void STableTreeView::TreeView_BuildViewColumnMenu(FMenuBuilder& MenuBuilder)
 
 void STableTreeView::TreeView_BuildExportMenu(FMenuBuilder& MenuBuilder)
 {
-	MenuBuilder.BeginSection("Export", LOCTEXT("ContextMenu_Section_Export", "Export"));
+	MenuBuilder.SetSearchable(false);
+
+	MenuBuilder.BeginSection("Export");
 	{
 		MenuBuilder.AddMenuEntry
 		(
