[Insights] Added a hint text as an overlay to the Groupings breadcrumb control. Also now the control allows clicking anywhere in order to open a drop down with options re add/edit groupings.

diff --git a/Engine/Source/Developer/TraceInsightsCore/Private/InsightsCore/Table/Widgets/STableTreeView.cpp b/Engine/Source/Developer/TraceInsightsCore/Private/InsightsCore/Table/Widgets/STableTreeView.cpp
--- a/Engine/Source/Developer/TraceInsightsCore/Private/InsightsCore/Table/Widgets/STableTreeView.cpp
+++ b/Engine/Source/Developer/TraceInsightsCore/Private/InsightsCore/Table/Widgets/STableTreeView.cpp
@@ -349,22 +349,54 @@ END_SLATE_FUNCTION_BUILD_OPTIMIZATION
 BEGIN_SLATE_FUNCTION_BUILD_OPTIMIZATION
 TSharedRef<SWidget> STableTreeView::ConstructHierarchyBreadcrumbTrail()
 {
-	TSharedRef<SWidget> Widget = SNew(SComboButton)
-		.HasDownArrow(false)
-		.ToolTipText(LOCTEXT("GroupingsTooltip", "Groupings Hierarchy"))
-		.ButtonContent()
+	TSharedRef<SWidget> Widget = SNew(SOverlay)
+
+		+ SOverlay::Slot()
+		[
+			SNew(SComboButton)
+			.HasDownArrow(false)
+			.ToolTipText(LOCTEXT("GroupingsTooltip", "Groupings Hierarchy"))
+			.OnGetMenuContent_Lambda([this]() -> TSharedRef<SWidget>
+				{
+					check(CurrentGroupings.Num() > 0);
+					return GetGroupingCrumbMenuContent(CurrentGroupings.Last());
+				})
+			.OnComboBoxOpened_Lambda([this]()
+				{
+					check(CurrentGroupings.Num() > 0);
+					OnGroupingCrumbClicked(CurrentGroupings.Last());
+				})
+			.ButtonContent()
+			[
+				SNew(SBox)
+				.HAlign(HAlign_Left)
+				[
+					SAssignNew(GroupingBreadcrumbTrail, SBreadcrumbTrail<TSharedPtr<FTreeNodeGrouping>>)
+					.ButtonStyle(FAppStyle::Get(), "SimpleButton")
+					.TextStyle(FAppStyle::Get(), "NormalText")
+					.DelimiterImage(FAppStyle::Get().GetBrush("Icons.ChevronRight"))
+					.ButtonContentPadding(FMargin(3.0f, 2.0f))
+					.GetCrumbMenuContent(this, &STableTreeView::GetGroupingCrumbMenuContent)
+					.OnCrumbClicked(this, &STableTreeView::OnGroupingCrumbClicked)
+					.ToolTipText(LOCTEXT("GroupingCrumbTooltip", "Click a grouping to remove the additional groupings.\nClick a separator to add/insert a new grouping or to change the current groupings."))
+				]
+			]
+		]
+
+		+ SOverlay::Slot()
 		[
 			SNew(SBox)
-			.HAlign(HAlign_Left)
+			.HAlign(HAlign_Right)
+			.VAlign(VAlign_Center)
+			.Padding(FMargin(0.0f, 0.0f, 6.0f, 0.0f))
+			.Visibility_Lambda([this]() ->EVisibility
+				{
+					return (CurrentGroupings.Num() == 1) ? EVisibility::HitTestInvisible : EVisibility::Collapsed;
+				})
 			[
-				SAssignNew(GroupingBreadcrumbTrail, SBreadcrumbTrail<TSharedPtr<FTreeNodeGrouping>>)
-				.ButtonStyle(FAppStyle::Get(), "SimpleButton")
-				.TextStyle(FAppStyle::Get(), "NormalText")
-				.DelimiterImage(FAppStyle::Get().GetBrush("Icons.ChevronRight"))
-				.ButtonContentPadding(FMargin(3.0f, 2.0f))
-				.GetCrumbMenuContent(this, &STableTreeView::GetGroupingCrumbMenuContent)
-				.OnCrumbClicked(this, &STableTreeView::OnGroupingCrumbClicked)
-				.ToolTipText(LOCTEXT("GroupingCrumbTooltip", "Click a grouping to remove the additional groupings.\nClick a separator to add/insert a new grouping or to change the current groupings."))
+				SNew(STextBlock)
+				.ColorAndOpacity(FSlateColor(EStyleColor::White25))
+				.Text(LOCTEXT("GroupingsHintText", "Click to add or edit groupings"))
 			]
 		];
 
