[Insights] Increased version of Unreal Insights to v1.07, for the UE 5.6 Release.

diff --git a/Engine/Source/Developer/TraceInsightsCore/Public/InsightsCore/Version.h b/Engine/Source/Developer/TraceInsightsCore/Public/InsightsCore/Version.h
--- a/Engine/Source/Developer/TraceInsightsCore/Public/InsightsCore/Version.h
+++ b/Engine/Source/Developer/TraceInsightsCore/Public/InsightsCore/Version.h
@@ -5,10 +5,10 @@
 // See https://semver.org/
 
 #define UNREAL_INSIGHTS_VERSION_MAJOR 1
-#define UNREAL_INSIGHTS_VERSION_MINOR 6
+#define UNREAL_INSIGHTS_VERSION_MINOR 7
 #define UNREAL_INSIGHTS_VERSION_PATCH 0
 
-#define UNREAL_INSIGHTS_VERSION_STRING "1.06"
+#define UNREAL_INSIGHTS_VERSION_STRING "1.07"
 
 /**
  * Extra identifier string (like: "alpha", "beta.1", "beta.2", "test", "featureX", etc.).
