[Insights] Table/Tree Views: Fixed Search functionality to use the display name instead of item FName.

diff --git a/Engine/Source/Developer/TraceInsightsCore/Private/InsightsCore/Table/Widgets/STableTreeView.cpp b/Engine/Source/Developer/TraceInsightsCore/Private/InsightsCore/Table/Widgets/STableTreeView.cpp
--- a/Engine/Source/Developer/TraceInsightsCore/Private/InsightsCore/Table/Widgets/STableTreeView.cpp
+++ b/Engine/Source/Developer/TraceInsightsCore/Private/InsightsCore/Table/Widgets/STableTreeView.cpp
@@ -545,7 +545,7 @@ TSharedPtr<SWidget> STableTreeView::TreeView_GetMenuContent()
 	}
 	else if (NumSelectedNodes == 1)
 	{
-		FString ItemName = SelectedNode->GetName().ToString();
+		FString ItemName = SelectedNode->GetDisplayName().ToString();
 		const int32 MaxStringLen = 64;
 		if (ItemName.Len() > MaxStringLen)
 		{
@@ -1327,7 +1327,7 @@ void STableTreeView::InitHierarchyFiltering()
 
 void STableTreeView::HandleItemToStringArray(const FTableTreeNodePtr& FTableTreeNodePtr, TArray<FString>& OutSearchStrings)
 {
-	OutSearchStrings.Add(FTableTreeNodePtr->GetName().ToString());
+	OutSearchStrings.Add(FTableTreeNodePtr->GetDisplayName().ToString());
 }
 
 ////////////////////////////////////////////////////////////////////////////////////////////////////
