Fix main window resizing when opening the filter configurator.

diff --git a/Engine/Source/Developer/TraceInsights/Private/Insights/Widgets/STimingView.cpp b/Engine/Source/Developer/TraceInsights/Private/Insights/Widgets/STimingView.cpp
--- a/Engine/Source/Developer/TraceInsights/Private/Insights/Widgets/STimingView.cpp
+++ b/Engine/Source/Developer/TraceInsights/Private/Insights/Widgets/STimingView.cpp
@@ -5227,9 +5227,12 @@ TSharedRef<SDockTab> STimingView::SpawnQuickFindTab(const FSpawnTabArgs& Args)
 	if (OwnerWindow.IsValid() && OwnerWindow != FSlateApplication::Get().FindWidgetWindow(SharedThis(this)))
 	{
 		TSharedPtr<SWindow> TopmostAncestor = OwnerWindow->GetTopmostAncestor();
-		const FVector2D DPIProbePoint = TopmostAncestor->GetPositionInScreen() + FVector2D(10.0, 10.0);
-		const float LocalDPIScaleFactor = FPlatformApplicationMisc::GetDPIScaleFactorAtPoint(static_cast<float>(DPIProbePoint.X), static_cast<float>(DPIProbePoint.Y));
-		OwnerWindow->Resize(FVector2D(600 * LocalDPIScaleFactor, 400 * LocalDPIScaleFactor));
+		if (TopmostAncestor != OwnerWindow)
+		{
+			const FVector2D DPIProbePoint = TopmostAncestor->GetPositionInScreen() + FVector2D(10.0, 10.0);
+			const float LocalDPIScaleFactor = FPlatformApplicationMisc::GetDPIScaleFactorAtPoint(static_cast<float>(DPIProbePoint.X), static_cast<float>(DPIProbePoint.Y));
+			OwnerWindow->Resize(FVector2D(600 * LocalDPIScaleFactor, 400 * LocalDPIScaleFactor));
+		}
 	}
 
 	if (!QuickFindWidgetSharedPtr.IsValid())
diff --git a/Engine/Source/Developer/TraceInsightsCore/Private/InsightsCore/Filter/ViewModels/Filters.cpp b/Engine/Source/Developer/TraceInsightsCore/Private/InsightsCore/Filter/ViewModels/Filters.cpp
--- a/Engine/Source/Developer/TraceInsightsCore/Private/InsightsCore/Filter/ViewModels/Filters.cpp
+++ b/Engine/Source/Developer/TraceInsightsCore/Private/InsightsCore/Filter/ViewModels/Filters.cpp
@@ -2,6 +2,7 @@
 
 #include "InsightsCore/Filter/ViewModels/Filters.h"
 
+#include "Framework/Application/SlateApplication.h"
 #include "HAL/PlatformApplicationMisc.h"
 #include "Widgets/Docking/SDockTab.h"
 
@@ -96,9 +97,10 @@ FFilterService::~FFilterService()
 
 ////////////////////////////////////////////////////////////////////////////////////////////////////
 
-TSharedPtr<SWidget> FFilterService::CreateFilterConfiguratorWidget(TSharedPtr<FFilterConfigurator> FilterConfiguratorViewModel)
+TSharedPtr<SWidget> FFilterService::CreateFilterConfiguratorWidget(TSharedPtr<FFilterConfigurator> FilterConfiguratorViewModel, TWeakPtr<SWidget> ParentWidget)
 {
 	SAssignNew(PendingWidget, SAdvancedFilter, FilterConfiguratorViewModel);
+	PendingParentWidget = ParentWidget;
 
 	if (FGlobalTabmanager::Get()->HasTabSpawner(FilterConfiguratorTabId))
 	{
@@ -115,11 +117,17 @@ TSharedRef<SDockTab> FFilterService::SpawnTab(const FSpawnTabArgs& Args)
 	const TSharedRef<SDockTab> DockTab = SNew(SDockTab)
 		.TabRole(ETabRole::NomadTab);
 
+	TSharedPtr<SWidget> ParentWidget = PendingParentWidget.Pin();
 	const TSharedPtr<SWindow>& OwnerWindow = Args.GetOwnerWindow();
-	if (OwnerWindow.IsValid())
+	if (OwnerWindow.IsValid() && ParentWidget.IsValid() && OwnerWindow != FSlateApplication::Get().FindWidgetWindow(ParentWidget.ToSharedRef()))
 	{
-		const float DPIScaleFactor = FPlatformApplicationMisc::GetDPIScaleFactorAtPoint(10.0f, 10.0f);
-		OwnerWindow->Resize(FVector2D(600 * DPIScaleFactor, 400 * DPIScaleFactor));
+		TSharedPtr<SWindow> TopmostAncestor = OwnerWindow->GetTopmostAncestor();
+		if (TopmostAncestor != OwnerWindow)
+		{
+			const FVector2D DPIProbePoint = TopmostAncestor->GetPositionInScreen() + FVector2D(10.0, 10.0);
+			const float LocalDPIScaleFactor = FPlatformApplicationMisc::GetDPIScaleFactorAtPoint(static_cast<float>(DPIProbePoint.X), static_cast<float>(DPIProbePoint.Y));
+			OwnerWindow->Resize(FVector2D(600 * LocalDPIScaleFactor, 400 * LocalDPIScaleFactor));
+		}
 	}
 
 	check(PendingWidget.IsValid());
@@ -127,6 +135,7 @@ TSharedRef<SDockTab> FFilterService::SpawnTab(const FSpawnTabArgs& Args)
 	PendingWidget->SetParentTab(DockTab);
 
 	PendingWidget = nullptr;
+	PendingParentWidget = nullptr;
 	return DockTab;
 }
 
diff --git a/Engine/Source/Developer/TraceInsightsCore/Private/InsightsCore/Table/Widgets/STableTreeView.cpp b/Engine/Source/Developer/TraceInsightsCore/Private/InsightsCore/Table/Widgets/STableTreeView.cpp
--- a/Engine/Source/Developer/TraceInsightsCore/Private/InsightsCore/Table/Widgets/STableTreeView.cpp
+++ b/Engine/Source/Developer/TraceInsightsCore/Private/InsightsCore/Table/Widgets/STableTreeView.cpp
@@ -1292,7 +1292,7 @@ FReply STableTreeView::FilterConfigurator_OnClicked()
 		CurrentAsyncOpFilterConfigurator = new FFilterConfigurator(*FilterConfigurator);
 	}
 
-	FFilterService::Get()->CreateFilterConfiguratorWidget(FilterConfigurator);
+	FFilterService::Get()->CreateFilterConfiguratorWidget(FilterConfigurator, SharedThis(this));
 
 	return FReply::Handled();
 }
diff --git a/Engine/Source/Developer/TraceInsightsCore/Public/InsightsCore/Filter/ViewModels/Filters.h b/Engine/Source/Developer/TraceInsightsCore/Public/InsightsCore/Filter/ViewModels/Filters.h
--- a/Engine/Source/Developer/TraceInsightsCore/Public/InsightsCore/Filter/ViewModels/Filters.h
+++ b/Engine/Source/Developer/TraceInsightsCore/Public/InsightsCore/Filter/ViewModels/Filters.h
@@ -355,13 +355,14 @@ public:
 	TSharedPtr<TArray<TSharedPtr<IFilterOperator>>> GetIntegerOperators() const { return FilterStorage.GetIntegerOperators(); }
 	TSharedPtr<TArray<TSharedPtr<IFilterOperator>>> GetStringOperators() const { return FilterStorage.GetStringOperators(); }
 
-	TSharedPtr<SWidget> CreateFilterConfiguratorWidget(TSharedPtr<class FFilterConfigurator> FilterConfiguratorViewModel);
+	TSharedPtr<SWidget> CreateFilterConfiguratorWidget(TSharedPtr<class FFilterConfigurator> FilterConfiguratorViewModel, TWeakPtr<SWidget> ParentWidget = nullptr);
 
 private:
 	static const FName FilterConfiguratorTabId;
 	static TSharedPtr<FFilterService> Instance;
 	FFilterStorage FilterStorage;
 	TSharedPtr<class SAdvancedFilter> PendingWidget;
+	TWeakPtr<SWidget> PendingParentWidget;
 };
 
 ////////////////////////////////////////////////////////////////////////////////////////////////////
