[Insights] Table Tree Views: Fixed node expansion not working as intended in few cases (ex. for Assets tag set in Memory Tags).

diff --git a/Engine/Source/Developer/TraceInsightsCore/Private/InsightsCore/Table/Widgets/STableTreeView.cpp b/Engine/Source/Developer/TraceInsightsCore/Private/InsightsCore/Table/Widgets/STableTreeView.cpp
--- a/Engine/Source/Developer/TraceInsightsCore/Private/InsightsCore/Table/Widgets/STableTreeView.cpp
+++ b/Engine/Source/Developer/TraceInsightsCore/Private/InsightsCore/Table/Widgets/STableTreeView.cpp
@@ -1424,37 +1424,9 @@ void STableTreeView::ApplyHierarchyFiltering()
 	// Cannot call TreeView functions from other threads than MainThread and SlateThread.
 	if (!bRunInAsyncMode)
 	{
-		// Only expand nodes if we have a text filter.
-		const bool bNonEmptyTextFilter = !TextFilter->GetRawFilterText().IsEmpty();
-		if (bNonEmptyTextFilter)
-		{
-			if (!bExpansionSaved)
-			{
-				ExpandedNodes.Empty();
-				TreeView->GetExpandedItems(ExpandedNodes);
-				bExpansionSaved = true;
-			}
+		Root->SetExpansion(true);
+		ApplyNodeExpansion(Root);
 
-			for (const FTableTreeNodePtr& GroupPtr : FilteredGroupNodes)
-			{
-				TreeView->SetItemExpansion(GroupPtr, GroupPtr->IsExpanded());
-			}
-		}
-		else
-		{
-			if (bExpansionSaved)
-			{
-				// Restore previously expanded nodes when the text filter is disabled.
-				TreeView->ClearExpandedItems();
-				for (auto It = ExpandedNodes.CreateConstIterator(); It; ++It)
-				{
-					TreeView->SetItemExpansion(*It, true);
-				}
-				bExpansionSaved = false;
-			}
-		}
-
-		// Request tree refresh
 		TreeView->RequestTreeRefresh();
 	}
 
@@ -1522,11 +1494,8 @@ bool STableTreeView::ApplyHierarchyFilteringRec(FTableTreeNodePtr NodePtr)
 		const bool bIsGroupNodeVisible = bIsNodeVisible || NumVisibleChildren > 0;
 		if (bIsGroupNodeVisible)
 		{
-			if (bRunInAsyncMode)
-			{
-				NodesToExpand.Add(NodePtr);
-			}
-			else
+			NodePtr->SetExpansion(true);
+			if (!bRunInAsyncMode)
 			{
 				TreeView->SetItemExpansion(NodePtr, true);
 			}
@@ -1551,7 +1520,7 @@ bool STableTreeView::MakeSubtreeVisible(FTableTreeNodePtr NodePtr, bool bFilterI
 
 		NodePtr->ClearFilteredChildren(bFilterIsEmpty ? GroupChildren.Num() : 0);
 
-		bool bShouldExpand = false;
+		bool bShouldExpand = NodePtr->IsExpanded();
 		for (const FBaseTreeNodePtr& GroupChildNodePtr : GroupChildren)
 		{
 			if (AsyncOperationProgress.ShouldCancelAsyncOp())
@@ -1562,18 +1531,10 @@ bool STableTreeView::MakeSubtreeVisible(FTableTreeNodePtr NodePtr, bool bFilterI
 			const FTableTreeNodePtr& ChildNodePtr = StaticCastSharedPtr<FTableTreeNode>(GroupChildNodePtr);
 			bShouldExpand |= MakeSubtreeVisible(ChildNodePtr, bFilterIsEmpty);
 			NodePtr->AddFilteredChild(ChildNodePtr);
-			NodePtr->SetExpansion(true);
-
-			if (bShouldExpand)
+			NodePtr->SetExpansion(bShouldExpand);
+			if (!bRunInAsyncMode)
 			{
-				if (bRunInAsyncMode)
-				{
-					NodesToExpand.Add(NodePtr);
-				}
-				else
-				{
-					TreeView->SetItemExpansion(NodePtr, true);
-				}
+				TreeView->SetItemExpansion(NodePtr, bShouldExpand);
 			}
 		}
 
@@ -1631,6 +1592,14 @@ void STableTreeView::TreeView_OnSelectionChanged(FTableTreeNodePtr SelectedItem,
 
 void STableTreeView::TreeView_OnExpansionChanged(FTableTreeNodePtr TreeNode, bool bShouldBeExpanded)
 {
+	if (TreeNode->IsExpanded() != bShouldBeExpanded)
+	{
+		TreeNode->SetExpansion(bShouldBeExpanded);
+		if (bShouldBeExpanded)
+		{
+			ApplyNodeExpansion(TreeNode);
+		}
+	}
 }
 
 ////////////////////////////////////////////////////////////////////////////////////////////////////
@@ -3368,13 +3337,8 @@ void STableTreeView::OnPreAsyncUpdate()
 	AsyncUpdateStopwatch.Restart();
 	bIsUpdateRunning = true;
 
-	if (ExpandedNodes.IsEmpty())
-	{
-		TreeView->GetExpandedItems(ExpandedNodes);
-	}
-
 	TreeView->SetTreeItemsSource(&DummyGroupNodes);
-	TreeView_Refresh();
+	TreeView->RequestTreeRefresh();
 
 	DispatchEvent = FGraphEvent::CreateGraphEvent();
 }
@@ -3391,29 +3355,17 @@ void STableTreeView::OnPostAsyncUpdate()
 	{
 		TreeView->SetTreeItemsSource(&FilteredGroupNodes);
 
-		TreeView->ClearExpandedItems();
-
-		// Grouping can result in old group nodes no longer existing in the tree, so we don't keep the expanded list.
-		if (!HasInProgressAsyncOperation(EAsyncOperationType::Grouping))
-		{
-			for (auto It = ExpandedNodes.CreateConstIterator(); It; ++It)
-			{
-				TreeView->SetItemExpansion(*It, true);
-			}
-		}
-
-		for (FTableTreeNodePtr& Node : NodesToExpand)
-		{
-			TreeView->SetItemExpansion(Node, true);
-		}
-		NodesToExpand.Empty();
-		ExpandedNodes.Empty();
-
 		// Expand each group node on the first few depths (if it doesn't have too many children).
 		SetExpandValueForChildGroups(Root.Get(), MaxNodesToAutoExpand, MaxDepthToAutoExpand, true);
 
+		// Update node expansion for the tree view.
+		TreeView->ClearExpandedItems();
+		Root->SetExpansion(true);
+		ApplyNodeExpansion(Root);
+
+		TreeView->RequestTreeRefresh();
+
 		ClearInProgressAsyncOperations();
-		TreeView_Refresh();
 	}
 
 	AsyncOperationProgress.Reset();
@@ -3440,8 +3392,33 @@ void STableTreeView::CountNumNodesPerDepthRec(FBaseTreeNode* InRoot, TArray<int3
 
 ////////////////////////////////////////////////////////////////////////////////////////////////////
 
+void STableTreeView::ApplyNodeExpansion(const FTableTreeNodePtr& InNode)
+{
+	if (InNode->IsGroup())
+	{
+		const bool bIsNodeExpanded = InNode->IsExpanded();
+
+		TreeView->SetItemExpansion(InNode, bIsNodeExpanded);
+
+		if (bIsNodeExpanded)
+		{
+			for (const FBaseTreeNodePtr& ChildNode : InNode->GetChildren())
+			{
+				ApplyNodeExpansion(StaticCastSharedPtr<FTableTreeNode>(ChildNode));
+			}
+		}
+	}
+}
+
+////////////////////////////////////////////////////////////////////////////////////////////////////
+
 void STableTreeView::SetExpandValueForChildGroups(FBaseTreeNode* InRoot, int32 InMaxNodesToExpand, int32 InMaxDepthToExpand, bool InValue)
 {
+	if (!InRoot)
+	{
+		return;
+	}
+
 	TArray<int32> NumNodesPerDepth;
 	NumNodesPerDepth.AddDefaulted(InMaxDepthToExpand + 1);
 	CountNumNodesPerDepthRec(InRoot, NumNodesPerDepth, 0, InMaxDepthToExpand, InMaxNodesToExpand);
@@ -3475,7 +3452,6 @@ void STableTreeView::SetExpandValueForChildGroupsRec(FBaseTreeNode* InRoot, int3
 		if (Node->IsGroup())
 		{
 			Node->SetExpansion(InValue);
-			TreeView->SetItemExpansion(StaticCastSharedPtr<FTableTreeNode>(Node), InValue);
 
 			if (InDepth < InMaxDepth)
 			{
@@ -3685,16 +3661,17 @@ void STableTreeView::ContextMenu_ExpandSubtree_Execute()
 		return;
 	}
 
-	for (const FTableTreeNodePtr& Node : SelectedNodes)
+	for (FTableTreeNodePtr& Node : SelectedNodes)
 	{
 		if (Node->IsGroup())
 		{
 			Node->SetExpansion(true);
-			TreeView->SetItemExpansion(Node, true);
 			SetExpandValueForChildGroups((FBaseTreeNode*)Node.Get(), MaxNodesToExpand, MaxDepthToExpand, true);
 		}
 	}
 
+	TreeView->ClearExpandedItems();
+	ApplyNodeExpansion(Root);
 	TreeView->RequestTreeRefresh();
 }
 
@@ -3702,21 +3679,7 @@ void STableTreeView::ContextMenu_ExpandSubtree_Execute()
 
 bool STableTreeView::ContextMenu_CollapseSubtree_CanExecute() const
 {
-	TArray<FTableTreeNodePtr> SelectedNodes;
-	if (!TreeView->GetSelectedItems(SelectedNodes))
-	{
-		return false;
-	}
-
-	for (const FTableTreeNodePtr& Node : SelectedNodes)
-	{
-		if (Node->IsGroup() && Node->GetFilteredChildren().Num() > 0 && TreeView->IsItemExpanded(Node))
-		{
-			return true;
-		}
-	}
-
-	return false;
+	return true;
 }
 
 ////////////////////////////////////////////////////////////////////////////////////////////////////
@@ -3729,16 +3692,17 @@ void STableTreeView::ContextMenu_CollapseSubtree_Execute()
 		return;
 	}
 
-	for (const FTableTreeNodePtr& Node : SelectedNodes)
+	for (FTableTreeNodePtr& Node : SelectedNodes)
 	{
-		if (Node->IsGroup() && TreeView->IsItemExpanded(Node))
+		if (Node->IsGroup())
 		{
 			Node->SetExpansion(false);
-			TreeView->SetItemExpansion(Node, false);
 			SetExpandValueForChildGroups((FBaseTreeNode*)Node.Get(), MaxNodesToExpand, MaxDepthToExpand, false);
 		}
 	}
 
+	TreeView->ClearExpandedItems();
+	ApplyNodeExpansion(Root);
 	TreeView->RequestTreeRefresh();
 }
 
diff --git a/Engine/Source/Developer/TraceInsightsCore/Public/InsightsCore/Table/Widgets/STableTreeView.h b/Engine/Source/Developer/TraceInsightsCore/Public/InsightsCore/Table/Widgets/STableTreeView.h
--- a/Engine/Source/Developer/TraceInsightsCore/Public/InsightsCore/Table/Widgets/STableTreeView.h
+++ b/Engine/Source/Developer/TraceInsightsCore/Public/InsightsCore/Table/Widgets/STableTreeView.h
@@ -453,6 +453,7 @@ protected:
 	////////////////////////////////////////////////////////////////////////////////////////////////////
 
 	void CountNumNodesPerDepthRec(FBaseTreeNode* InRoot, TArray<int32>& InOutNumNodesPerDepth, int32 InDepth, int32 InMaxDepth, int32 InMaxNodes) const;
+	void ApplyNodeExpansion(const FTableTreeNodePtr& InNode);
 	void SetExpandValueForChildGroups(FBaseTreeNode* InRoot, int32 InMaxExpandedNodes, int32 MaxDepthToExpand, bool InValue);
 	void SetExpandValueForChildGroupsRec(FBaseTreeNode* InRoot, int32 InDepth, int32 InMaxDepth, bool InValue);
 
@@ -548,12 +549,6 @@ protected:
 	/** A filtered array of group nodes to be displayed in the tree widget. */
 	TArray<FTableTreeNodePtr> FilteredGroupNodes;
 
-	/** Currently expanded group nodes. */
-	TSet<FTableTreeNodePtr> ExpandedNodes;
-
-	/** If true, the expanded nodes have been saved before applying a text filter. Only used when bRunInAsyncMode==false. */
-	bool bExpansionSaved = false;
-
 	int32 MaxNodesToAutoExpand = 1000;
 	int32 MaxDepthToAutoExpand = 4;
 	int32 MaxNodesToExpand = 1000000;
@@ -635,7 +630,6 @@ protected:
 	FAsyncOperationProgress AsyncOperationProgress;
 	FGraphEventRef DispatchEvent;
 	TArray<TSharedPtr<FTableTaskInfo>> DataTaskInfos;
-	TArray<FTableTreeNodePtr> NodesToExpand;
 
 	//////////////////////////////////////////////////
 
