[Insights] Fixed groupings to not make a column visible if it was already hidden when the respective grouping was applied.

diff --git a/Engine/Source/Developer/TraceInsightsCore/Private/InsightsCore/Table/Widgets/STableTreeView.cpp b/Engine/Source/Developer/TraceInsightsCore/Private/InsightsCore/Table/Widgets/STableTreeView.cpp
--- a/Engine/Source/Developer/TraceInsightsCore/Private/InsightsCore/Table/Widgets/STableTreeView.cpp
+++ b/Engine/Source/Developer/TraceInsightsCore/Private/InsightsCore/Table/Widgets/STableTreeView.cpp
@@ -2042,15 +2042,12 @@ void STableTreeView::InternalCreateGroupings()
 
 void STableTreeView::PreChangeGroupings()
 {
-	for (TSharedPtr<FTreeNodeGrouping>& GroupingPtr : CurrentGroupings)
+	// Show columns used in previous groupings.
+	for (const FName& ColumnId : GroupingHiddenColumns)
 	{
-		const FName& ColumnId = GroupingPtr->GetColumnId();
-		if (ColumnId != NAME_None)
-		{
-			// Show columns used in previous groupings.
-			ShowColumn(ColumnId);
-		}
+		ShowColumn(ColumnId);
 	}
+	GroupingHiddenColumns.Reset();
 }
 
 ////////////////////////////////////////////////////////////////////////////////////////////////////
@@ -2092,7 +2089,12 @@ void STableTreeView::PostChangeGroupings()
 			}
 
 			// Hide columns used in groupings.
-			HideColumn(ColumnId);
+			FTableColumn& Column = *Table->FindColumnChecked(ColumnId);
+			if (Column.IsVisible())
+			{
+				GroupingHiddenColumns.Add(ColumnId);
+				HideColumn(Column);
+			}
 		}
 
 		if (bAdjustHierarchyColumnName)
diff --git a/Engine/Source/Developer/TraceInsightsCore/Public/InsightsCore/Table/Widgets/STableTreeView.h b/Engine/Source/Developer/TraceInsightsCore/Public/InsightsCore/Table/Widgets/STableTreeView.h
--- a/Engine/Source/Developer/TraceInsightsCore/Public/InsightsCore/Table/Widgets/STableTreeView.h
+++ b/Engine/Source/Developer/TraceInsightsCore/Public/InsightsCore/Table/Widgets/STableTreeView.h
@@ -593,6 +593,9 @@ protected:
 	/** The groupings actually used in the Grouping async task. */
 	TArray<TSharedPtr<FTreeNodeGrouping>> CurrentAsyncOpGroupings;
 
+	/** The columns hidden when grouping is applied. */
+	TSet<FName> GroupingHiddenColumns;
+
 	//////////////////////////////////////////////////
 	// Sorting
 
