[Insights]
* Table/Tree Views: Added support for Color to groupings (each breadcrumb grouping can have an assigned color).
* Table/Tree Views: Assigned pseudo-random colors to "Unique Values" and to "Path Breakdown" groupings.
* Table/Tree Views: Moved the list of available "Unique Values" groupings and "Path Breakdown" groupings into sub-menus, if more than 3.
* Memory Insights: Assigned colors for the custom groupings used in MemoryTags and MemAllocs tables (based on the color of group nodes they create).

diff --git a/Engine/Source/Developer/TraceInsights/Private/Insights/MemoryProfiler/ViewModels/MemAllocGroupingByCallstack.cpp b/Engine/Source/Developer/TraceInsights/Private/Insights/MemoryProfiler/ViewModels/MemAllocGroupingByCallstack.cpp
--- a/Engine/Source/Developer/TraceInsights/Private/Insights/MemoryProfiler/ViewModels/MemAllocGroupingByCallstack.cpp
+++ b/Engine/Source/Developer/TraceInsights/Private/Insights/MemoryProfiler/ViewModels/MemAllocGroupingByCallstack.cpp
@@ -68,11 +68,12 @@ INSIGHTS_IMPLEMENT_RTTI(FMemAllocGroupingByCallstack)
 
 FMemAllocGroupingByCallstack::FMemAllocGroupingByCallstack(const FText& InShortName, const FText& InTitleName, const FText& InDescription, const FSlateBrush* InIcon, bool bInIsAllocCallstack, bool bInIsInverted, bool bInIsGroupingByFunction)
 	: FTreeNodeGrouping(InShortName, InTitleName, InDescription, InIcon)
-		, bIsAllocCallstack(bInIsAllocCallstack)
-		, bIsInverted(bInIsInverted)
-		, bIsGroupingByFunction(bInIsGroupingByFunction)
-		, bShouldSkipFilteredFrames(bIsInverted)
+	, bIsAllocCallstack(bInIsAllocCallstack)
+	, bIsInverted(bInIsInverted)
+	, bIsGroupingByFunction(bInIsGroupingByFunction)
+	, bShouldSkipFilteredFrames(bIsInverted)
 {
+	SetColor(FLinearColor(0.5f, 0.75f, 1.0f, 1.0f));
 }
 
 ////////////////////////////////////////////////////////////////////////////////////////////////////
@@ -83,18 +84,17 @@ FMemAllocGroupingByCallstack::FMemAllocGroupingByCallstack(bool bInIsAllocCallst
 			? (bInIsInverted ? LOCTEXT("Grouping_ByCallstack2_ShortName", "Inverted Alloc Callstack")
 							 : LOCTEXT("Grouping_ByCallstack1_ShortName", "Alloc Callstack"))
 			: (bInIsInverted ? LOCTEXT("Grouping_ByCallstack4_ShortName", "Inverted Free Callstack")
-							 : LOCTEXT("Grouping_ByCallstack3_ShortName", "Free Callstack"))
-		, bInIsAllocCallstack
+							 : LOCTEXT("Grouping_ByCallstack3_ShortName", "Free Callstack")),
+		bInIsAllocCallstack
 			? (bInIsInverted ? LOCTEXT("Grouping_ByCallstack2_TitleName", "By Inverted Alloc Callstack")
 							 : LOCTEXT("Grouping_ByCallstack1_TitleName", "By Alloc Callstack"))
 			: (bInIsInverted ? LOCTEXT("Grouping_ByCallstack4_TitleName", "By Inverted Free Callstack")
-							 : LOCTEXT("Grouping_ByCallstack3_TitleName", "By Free Callstack"))
-		,LOCTEXT("Grouping_Callstack_Desc", "Creates a tree based on callstack of each allocation.")
-		,nullptr
-	, bInIsAllocCallstack
-	, bInIsInverted
-	, bInIsGroupingByFunction
-	)
+							 : LOCTEXT("Grouping_ByCallstack3_TitleName", "By Free Callstack")),
+		LOCTEXT("Grouping_Callstack_Desc", "Creates a tree based on callstack of each allocation."),
+		nullptr,
+		bInIsAllocCallstack,
+		bInIsInverted,
+		bInIsGroupingByFunction)
 {
 }
 
diff --git a/Engine/Source/Developer/TraceInsights/Private/Insights/MemoryProfiler/ViewModels/MemAllocGroupingByHeap.cpp b/Engine/Source/Developer/TraceInsights/Private/Insights/MemoryProfiler/ViewModels/MemAllocGroupingByHeap.cpp
--- a/Engine/Source/Developer/TraceInsights/Private/Insights/MemoryProfiler/ViewModels/MemAllocGroupingByHeap.cpp
+++ b/Engine/Source/Developer/TraceInsights/Private/Insights/MemoryProfiler/ViewModels/MemAllocGroupingByHeap.cpp
@@ -31,6 +31,7 @@ FMemAllocGroupingByHeap::FMemAllocGroupingByHeap(const TraceServices::IAllocatio
 		nullptr)
 	, AllocProvider(InAllocProvider)
 {
+	SetColor(FLinearColor(1.0f, 0.45f, 0.6f, 1.0f));
 }
 
 ////////////////////////////////////////////////////////////////////////////////////////////////////
diff --git a/Engine/Source/Developer/TraceInsights/Private/Insights/MemoryProfiler/ViewModels/MemAllocGroupingBySwapPage.cpp b/Engine/Source/Developer/TraceInsights/Private/Insights/MemoryProfiler/ViewModels/MemAllocGroupingBySwapPage.cpp
--- a/Engine/Source/Developer/TraceInsights/Private/Insights/MemoryProfiler/ViewModels/MemAllocGroupingBySwapPage.cpp
+++ b/Engine/Source/Developer/TraceInsights/Private/Insights/MemoryProfiler/ViewModels/MemAllocGroupingBySwapPage.cpp
@@ -16,33 +16,6 @@
 namespace UE::Insights::MemoryProfiler
 {
 
-INSIGHTS_IMPLEMENT_RTTI(FMemAllocGroupingBySwapPage)
-
-////////////////////////////////////////////////////////////////////////////////////////////////////
-
-FMemAllocGroupingBySwapPage::FMemAllocGroupingBySwapPage(const TraceServices::IAllocationsProvider& InAllocProvider, bool bInIsInverted, bool bInIsGroupingByFunction)
-	: FMemAllocGroupingByCallstack(
- 		bInIsInverted ? 
-			LOCTEXT("Grouping_BySwapCallstack2_ShortName", "In-Swap Inverted Alloc Callstack") 
-			: LOCTEXT("Grouping_BySwapCallstack1_ShortName", "In-Swap Alloc Callstack"),
-		bInIsInverted ? 
-			LOCTEXT("Grouping_BySwapCallstack2_TitleName", "By Inverted In-Swap Inverted Alloc Callstack")
-			: LOCTEXT("Grouping_BySwapCallstack1_TitleName", "By In-Swap Alloc Callstack"),
-		LOCTEXT("Grouping_SwapCallstack_Desc", "Creates a tree based on callstack of each allocation which is in swap memory."),
-		nullptr
-	, true /*bInIsAllocCallstack*/
-	, bInIsInverted
-	, bInIsGroupingByFunction)
-	, AllocProvider(InAllocProvider)
-{
-}
-
-////////////////////////////////////////////////////////////////////////////////////////////////////
-
-FMemAllocGroupingBySwapPage::~FMemAllocGroupingBySwapPage()
-{
-}
-
 ////////////////////////////////////////////////////////////////////////////////////////////////////
 
 struct TMemoryPageMapKeyFuncs : BaseKeyFuncs<TPair<uint64, FTableTreeNodePtr>, uint64, false>
@@ -62,6 +35,8 @@ struct TMemoryPageMapKeyFuncs : BaseKeyFuncs<TPair<uint64, FTableTreeNodePtr>, u
 	}
 };
 
+////////////////////////////////////////////////////////////////////////////////////////////////////
+// FMemSwapPageTreeNode
 ////////////////////////////////////////////////////////////////////////////////////////////////////
 
 class FMemSwapPageTreeNode : public FTableTreeNode
@@ -89,6 +64,13 @@ public:
 		return FLinearColor(0.2f, 0.8f, 0.4f, 1.0f);
 	}
 };
+
+////////////////////////////////////////////////////////////////////////////////////////////////////
+
+INSIGHTS_IMPLEMENT_RTTI(FMemSwapPageTreeNode)
+
+////////////////////////////////////////////////////////////////////////////////////////////////////
+// FMemSwappedAllocsTreeNode
 ////////////////////////////////////////////////////////////////////////////////////////////////////
 
 class FMemSwappedAllocsTreeNode : public FMemAllocValueNode
@@ -127,9 +109,38 @@ public:
 
 ////////////////////////////////////////////////////////////////////////////////////////////////////
 
-INSIGHTS_IMPLEMENT_RTTI(FMemSwapPageTreeNode)
 INSIGHTS_IMPLEMENT_RTTI(FMemSwappedAllocsTreeNode)
 
+////////////////////////////////////////////////////////////////////////////////////////////////////
+// FMemAllocGroupingBySwapPage
+////////////////////////////////////////////////////////////////////////////////////////////////////
+
+INSIGHTS_IMPLEMENT_RTTI(FMemAllocGroupingBySwapPage)
+
+////////////////////////////////////////////////////////////////////////////////////////////////////
+
+FMemAllocGroupingBySwapPage::FMemAllocGroupingBySwapPage(const TraceServices::IAllocationsProvider& InAllocProvider, bool bInIsInverted, bool bInIsGroupingByFunction)
+	: FMemAllocGroupingByCallstack(
+		bInIsInverted ? LOCTEXT("Grouping_BySwapCallstack2_ShortName", "In-Swap Inverted Alloc Callstack")
+					  : LOCTEXT("Grouping_BySwapCallstack1_ShortName", "In-Swap Alloc Callstack"),
+		bInIsInverted ? LOCTEXT("Grouping_BySwapCallstack2_TitleName", "By In-Swap Inverted Alloc Callstack")
+					  : LOCTEXT("Grouping_BySwapCallstack1_TitleName", "By In-Swap Alloc Callstack"),
+		LOCTEXT("Grouping_SwapCallstack_Desc", "Creates a tree based on callstack of each allocation which is in swap memory."),
+		nullptr,
+		true /*bInIsAllocCallstack*/,
+		bInIsInverted,
+		bInIsGroupingByFunction)
+	, AllocProvider(InAllocProvider)
+{
+	SetColor(FLinearColor(0.2f, 0.8f, 0.4f, 1.0f));
+}
+
+////////////////////////////////////////////////////////////////////////////////////////////////////
+
+FMemAllocGroupingBySwapPage::~FMemAllocGroupingBySwapPage()
+{
+}
+
 ////////////////////////////////////////////////////////////////////////////////////////////////////
 
 void FMemAllocGroupingBySwapPage::GroupNodes(
@@ -216,9 +227,9 @@ FTableTreeNodePtr FMemAllocGroupingBySwapPage::ProcessNodeForSwapInfo(FTableTree
 		}
 		uint64 AllocSizeInPage = PageSize;
 
-		if (PageAddress < AllocStart) 
+		if (PageAddress < AllocStart)
 		{
-			// adjust size if allocation starts after page start 
+			// adjust size if allocation starts after page start
 			AllocSizeInPage -= AllocStart - PageAddress;
 		}
 		if (PageAddress + PageSize > AllocEnd)
@@ -227,7 +238,7 @@ FTableTreeNodePtr FMemAllocGroupingBySwapPage::ProcessNodeForSwapInfo(FTableTree
 			AllocSizeInPage -= PageAddress + PageSize - AllocEnd;
 		}
 
-		if(AllocSizeInPage)
+		if (AllocSizeInPage)
 		{
 			SwapPagesEncountered.Add(*SwapNodePtrPtr);
 		}
@@ -250,8 +261,6 @@ FTableTreeNodePtr FMemAllocGroupingBySwapPage::ProcessNodeForSwapInfo(FTableTree
 
 		if (TSharedPtr<FTable> SharedTable = InParentTable.Pin())
 		{
-
-
 			auto SwapEventsGroup = MakeShared<FMemSwapPageTreeNode>(FName(FString::Printf(TEXT("Swap events"))), InParentTable);
 			SwappedWrapper->AddChildAndSetParent(SwapEventsGroup);
 
diff --git a/Engine/Source/Developer/TraceInsights/Private/Insights/MemoryProfiler/ViewModels/MemAllocGroupingByTag.cpp b/Engine/Source/Developer/TraceInsights/Private/Insights/MemoryProfiler/ViewModels/MemAllocGroupingByTag.cpp
--- a/Engine/Source/Developer/TraceInsights/Private/Insights/MemoryProfiler/ViewModels/MemAllocGroupingByTag.cpp
+++ b/Engine/Source/Developer/TraceInsights/Private/Insights/MemoryProfiler/ViewModels/MemAllocGroupingByTag.cpp
@@ -37,6 +37,7 @@ FMemAllocGroupingByTag::FMemAllocGroupingByTag(const TraceServices::IAllocations
 		nullptr)
 	, TagProvider(InTagProvider)
 {
+	SetColor(FLinearColor(0.75f, 0.5f, 1.0f, 1.0f));
 }
 
 ////////////////////////////////////////////////////////////////////////////////////////////////////
diff --git a/Engine/Source/Developer/TraceInsights/Private/Insights/MemoryProfiler/ViewModels/MemTagBudgetGrouping.cpp b/Engine/Source/Developer/TraceInsights/Private/Insights/MemoryProfiler/ViewModels/MemTagBudgetGrouping.cpp
--- a/Engine/Source/Developer/TraceInsights/Private/Insights/MemoryProfiler/ViewModels/MemTagBudgetGrouping.cpp
+++ b/Engine/Source/Developer/TraceInsights/Private/Insights/MemoryProfiler/ViewModels/MemTagBudgetGrouping.cpp
@@ -36,6 +36,7 @@ FMemTagBudgetNodeGrouping::FMemTagBudgetNodeGrouping(TSharedPtr<SMemTagTreeView>
 		nullptr)
 	, MemTagTreeView(InMemTagTreeView.ToWeakPtr())
 {
+	SetColor(FLinearColor(0.75f, 0.5f, 1.0f, 1.0f));
 }
 
 ////////////////////////////////////////////////////////////////////////////////////////////////////
diff --git a/Engine/Source/Developer/TraceInsights/Private/Insights/MemoryProfiler/Widgets/SMemTagTreeView.cpp b/Engine/Source/Developer/TraceInsights/Private/Insights/MemoryProfiler/Widgets/SMemTagTreeView.cpp
--- a/Engine/Source/Developer/TraceInsights/Private/Insights/MemoryProfiler/Widgets/SMemTagTreeView.cpp
+++ b/Engine/Source/Developer/TraceInsights/Private/Insights/MemoryProfiler/Widgets/SMemTagTreeView.cpp
@@ -100,7 +100,7 @@ void SMemTagTreeView::ConstructHeaderArea(TSharedRef<SVerticalBox> InHostBox)
 				TagSetAndViewPresetWidget.ToSharedRef()
 			];
 	}
-	
+
 	TSharedPtr<SWidget> TimeMarkersWidget = ConstructTimeMarkers();
 	if (TimeMarkersWidget.IsValid())
 	{
@@ -953,6 +953,33 @@ void SMemTagTreeView::InternalCreateGroupings()
 {
 	STableTreeView::InternalCreateGroupings();
 
+	AvailableGroupings.RemoveAll(
+		[](TSharedPtr<FTreeNodeGrouping>& Grouping)
+		{
+			if (Grouping->Is<FTreeNodeGroupingByUniqueValue>())
+			{
+				const FName ColumnId = Grouping->As<FTreeNodeGroupingByUniqueValue>().GetColumnId();
+				return (ColumnId != FMemTagTableColumns::SampleCountColumnId &&
+						ColumnId != FMemTagTableColumns::TrackerColumnId);
+			}
+			else if (Grouping->Is<FTreeNodeGroupingByPathBreakdown>())
+			{
+				const FName ColumnId = Grouping->As<FTreeNodeGroupingByPathBreakdown>().GetColumnId();
+				return (ColumnId != FMemTagTableColumns::TagNameColumnId);
+			}
+			return false;
+		});
+
+	TSharedPtr<FTreeNodeGrouping>* BySampleCountUnique = AvailableGroupings.FindByPredicate([](TSharedPtr<FTreeNodeGrouping>& Grouping)
+		{
+			return Grouping->Is<FTreeNodeGroupingByUniqueValue>() &&
+				   Grouping->As<FTreeNodeGroupingByUniqueValue>().GetColumnId() == FMemTagTableColumns::SampleCountColumnId;
+		});
+	if (BySampleCountUnique != nullptr)
+	{
+		(*BySampleCountUnique)->SetColor(FLinearColor(0.6f, 0.5f, 0.1f, 1.0f));
+	}
+
 	int32 Index = 1; // after the Flat ("All") grouping
 
 	AvailableGroupings.Insert(MakeShared<FMemTagBudgetNodeGrouping>(SharedThis(this)), Index++);
@@ -1961,7 +1988,7 @@ void SMemTagTreeView::SetColorToNode(FMemTagNode& MemTagNode, FLinearColor Color
 	{
 		return;
 	}
-	
+
 	MemTag->SetColor(Color);
 	const FLinearColor BorderColor(FMath::Min(Color.R + 0.4f, 1.0f), FMath::Min(Color.G + 0.4f, 1.0f), FMath::Min(Color.B + 0.4f, 1.0f), 1.0f);
 
diff --git a/Engine/Source/Developer/TraceInsightsCore/Private/InsightsCore/Table/ViewModels/TreeNodeGrouping.cpp b/Engine/Source/Developer/TraceInsightsCore/Private/InsightsCore/Table/ViewModels/TreeNodeGrouping.cpp
--- a/Engine/Source/Developer/TraceInsightsCore/Private/InsightsCore/Table/ViewModels/TreeNodeGrouping.cpp
+++ b/Engine/Source/Developer/TraceInsightsCore/Private/InsightsCore/Table/ViewModels/TreeNodeGrouping.cpp
@@ -22,38 +22,21 @@ INSIGHTS_IMPLEMENT_RTTI(FTreeNodeGroupingByPathBreakdown)
 // FTreeNodeGrouping
 ////////////////////////////////////////////////////////////////////////////////////////////////////
 
-PRAGMA_DISABLE_DEPRECATION_WARNINGS
-FTreeNodeGrouping::FTreeNodeGrouping(const FText& InShortName, const FText& InTitleName, const FText& InDescription, const FSlateBrush* InIcon)
-	: ShortName(InShortName)
-	, TitleName(InTitleName)
-	, Description(InDescription)
-	, BrushName(NAME_None)
-	, Icon(InIcon)
+FTreeNodeGrouping::FTreeNodeGrouping()
+	: Color(FBaseTreeNode::GetDefaultColor(true))
 {
 }
-PRAGMA_ENABLE_DEPRECATION_WARNINGS
 
 ////////////////////////////////////////////////////////////////////////////////////////////////////
 
-PRAGMA_DISABLE_DEPRECATION_WARNINGS
-FTreeNodeGrouping::FTreeNodeGrouping(const FText& InShortName, const FText& InTitleName, const FText& InDescription, const FName InBrushName, const FSlateBrush* InIcon)
+FTreeNodeGrouping::FTreeNodeGrouping(const FText& InShortName, const FText& InTitleName, const FText& InDescription, const FSlateBrush* InIcon)
 	: ShortName(InShortName)
 	, TitleName(InTitleName)
 	, Description(InDescription)
-	, BrushName(InBrushName)
 	, Icon(InIcon)
+	, Color(FBaseTreeNode::GetDefaultColor(true))
 {
 }
-PRAGMA_ENABLE_DEPRECATION_WARNINGS
-
-////////////////////////////////////////////////////////////////////////////////////////////////////
-
-FName FTreeNodeGrouping::GetBrushName() const
-{
-	PRAGMA_DISABLE_DEPRECATION_WARNINGS
-	return BrushName;
-	PRAGMA_ENABLE_DEPRECATION_WARNINGS
-}
 
 ////////////////////////////////////////////////////////////////////////////////////////////////////
 
@@ -67,7 +50,7 @@ void FTreeNodeGrouping::GroupNodes(const TArray<FTableTreeNodePtr>& Nodes, FTabl
 	{
 		if (InAsyncOperationProgress.ShouldCancelAsyncOp())
 		{
-			return;
+			break;
 		}
 
 		if (NodePtr->IsGroup())
@@ -79,13 +62,21 @@ void FTreeNodeGrouping::GroupNodes(const TArray<FTableTreeNodePtr>& Nodes, FTabl
 		FTableTreeNodePtr GroupPtr = nullptr;
 
 		FTreeNodeGroupInfo GroupInfo = GetGroupForNode(NodePtr);
-		FTableTreeNodePtr* GroupPtrPtr = GroupMap.Find(GroupInfo.Name);
+		FTableTreeNodePtr* GroupPtrPtr = GroupMap.Find(GroupInfo.GetName());
 		if (!GroupPtrPtr)
 		{
-			GroupPtr = MakeShared<FTableTreeNode>(GroupInfo.Name, InParentTable);
-			GroupPtr->SetExpansion(GroupInfo.IsExpanded);
+			if (GroupInfo.IsCustom())
+			{
+				const FSlateBrush* GroupIcon = GroupInfo.GetIcon() != nullptr ? GroupInfo.GetIcon() : FBaseTreeNode::GetDefaultIcon(true);
+				GroupPtr = MakeShared<FCustomTableTreeNode>(GroupInfo.GetName(), InParentTable, GroupIcon, GroupInfo.GetIconColor(), GroupInfo.GetColor());
+			}
+			else
+			{
+				GroupPtr = MakeShared<FTableTreeNode>(GroupInfo.GetName(), InParentTable);
+			}
+			GroupPtr->SetExpansion(GroupInfo.IsExpanded());
 			ParentGroup.AddChildAndSetParent(GroupPtr);
-			GroupMap.Add(GroupInfo.Name, GroupPtr);
+			GroupMap.Add(GroupInfo.GetName(), GroupPtr);
 		}
 		else
 		{
@@ -96,6 +87,14 @@ void FTreeNodeGrouping::GroupNodes(const TArray<FTableTreeNodePtr>& Nodes, FTabl
 	}
 }
 
+////////////////////////////////////////////////////////////////////////////////////////////////////
+
+FLinearColor FTreeNodeGrouping::MakeDistinctColor(const FName& InName)
+{
+	uint32 Seed = GetTypeHash(InName.ToString());
+	return FLinearColor(static_cast<float>(Seed % 3600) / 10.0f, 0.7f, 1.0f, 1.0f).HSVToLinearRGB();
+}
+
 ////////////////////////////////////////////////////////////////////////////////////////////////////
 // FTreeNodeGroupingFlat
 ////////////////////////////////////////////////////////////////////////////////////////////////////
@@ -115,7 +114,7 @@ void FTreeNodeGroupingFlat::GroupNodes(const TArray<FTableTreeNodePtr>& Nodes, F
 {
 	ParentGroup.ClearChildren(1);
 
-	FTableTreeNodePtr GroupPtr = MakeShared<FTableTreeNode>(FName(TEXT("All")), InParentTable);
+	FTableTreeNodePtr GroupPtr = MakeShared<FCustomTableTreeNode>(FName(TEXT("All")), InParentTable, GetColor());
 	GroupPtr->SetExpansion(true);
 	ParentGroup.AddChildAndSetParent(GroupPtr);
 
@@ -124,7 +123,7 @@ void FTreeNodeGroupingFlat::GroupNodes(const TArray<FTableTreeNodePtr>& Nodes, F
 	{
 		if (InAsyncOperationProgress.ShouldCancelAsyncOp())
 		{
-			return;
+			break;
 		}
 
 		GroupPtr->AddChildAndSetParent(NodePtr);
@@ -143,6 +142,7 @@ FTreeNodeGroupingByUniqueValue::FTreeNodeGroupingByUniqueValue(TSharedRef<FTable
 		nullptr)
 	, ColumnRef(InColumnRef)
 {
+	SetColor(FTreeNodeGrouping::MakeDistinctColor(InColumnRef->GetId()));
 }
 
 ////////////////////////////////////////////////////////////////////////////////////////////////////
@@ -156,7 +156,7 @@ FTreeNodeGroupInfo FTreeNodeGroupingByUniqueValue::GetGroupForNode(const FBaseTr
 	{
 		GroupName = FStringView(GroupName.GetData(), NAME_SIZE - 1);
 	}
-	return { FName(GroupName, 0), false };
+	return FTreeNodeGroupInfo(FName(GroupName, 0), false, GetColor());
 }
 
 ////////////////////////////////////////////////////////////////////////////////////////////////////
@@ -197,7 +197,7 @@ void FTreeNodeGroupingByUniqueValueCString::GroupNodes(const TArray<FTableTreeNo
 	{
 		if (InAsyncOperationProgress.ShouldCancelAsyncOp())
 		{
-			return;
+			break;
 		}
 
 		if (NodePtr->IsGroup())
@@ -222,7 +222,7 @@ void FTreeNodeGroupingByUniqueValueCString::GroupNodes(const TArray<FTableTreeNo
 					GroupPtrPtr = GroupNameMap.Find(Value);
 					if (!GroupPtrPtr)
 					{
-						GroupPtr = MakeShared<FTableTreeNode>(GroupName, InParentTable);
+						GroupPtr = MakeShared<FCustomTableTreeNode>(GroupName, InParentTable, GetColor());
 						GroupPtr->SetExpansion(false);
 						ParentGroup.AddChildAndSetParent(GroupPtr);
 						GroupNameMap.Add(GroupName, GroupPtr);
@@ -316,6 +316,7 @@ FTreeNodeGroupingByPathBreakdown::FTreeNodeGroupingByPathBreakdown(TSharedRef<FT
 		nullptr)
 	, ColumnRef(InColumnRef)
 {
+	SetColor(FTreeNodeGrouping::MakeDistinctColor(InColumnRef->GetId()));
 }
 
 ////////////////////////////////////////////////////////////////////////////////////////////////////
@@ -357,7 +358,7 @@ void FTreeNodeGroupingByPathBreakdown::GroupNodes(
 	{
 		if (InAsyncOperationProgress.ShouldCancelAsyncOp())
 		{
-			return;
+			break;
 		}
 
 		if (NodePtr->IsGroup())
@@ -399,7 +400,7 @@ void FTreeNodeGroupingByPathBreakdown::GroupNodes(
 						}
 						else if (!DirName.IsNone())
 						{
-							FTableTreeNodePtr GroupPtr = MakeShared<FTableTreeNode>(DirName, InParentTable);
+							FTableTreeNodePtr GroupPtr = MakeShared<FCustomTableTreeNode>(DirName, InParentTable, GetColor());
 							GroupPtr->SetExpansion(false);
 							Directory->Node->AddChildAndSetParent(GroupPtr);
 							Directory = new FDirectory(DirName, Directory, GroupPtr.Get());
@@ -438,6 +439,7 @@ void FTreeNodeGroupingByPathBreakdown::GroupNodes(
 		}
 	}
 
+	// Clean up.
 	for (FDirectory* Dir : Directories)
 	{
 		delete Dir;
diff --git a/Engine/Source/Developer/TraceInsightsCore/Private/InsightsCore/Table/Widgets/STableTreeView.cpp b/Engine/Source/Developer/TraceInsightsCore/Private/InsightsCore/Table/Widgets/STableTreeView.cpp
--- a/Engine/Source/Developer/TraceInsightsCore/Private/InsightsCore/Table/Widgets/STableTreeView.cpp
+++ b/Engine/Source/Developer/TraceInsightsCore/Private/InsightsCore/Table/Widgets/STableTreeView.cpp
@@ -376,6 +376,7 @@ TSharedRef<SWidget> STableTreeView::ConstructHierarchyBreadcrumbTrail()
 					.TextStyle(FAppStyle::Get(), "NormalText")
 					.DelimiterImage(FAppStyle::Get().GetBrush("Icons.ChevronRight"))
 					.ButtonContentPadding(FMargin(3.0f, 2.0f))
+					.GetCrumbButtonContent(this, &STableTreeView::GetGroupingCrumbButtonContent)
 					.GetCrumbMenuContent(this, &STableTreeView::GetGroupingCrumbMenuContent)
 					.OnCrumbClicked(this, &STableTreeView::OnGroupingCrumbClicked)
 					.ToolTipText(LOCTEXT("GroupingCrumbTooltip", "Click a grouping to remove the additional groupings.\nClick a separator to add/insert a new grouping or to change the current groupings."))
@@ -2193,60 +2194,148 @@ void STableTreeView::OnGroupingCrumbClicked(const TSharedPtr<FTreeNodeGrouping>&
 
 ////////////////////////////////////////////////////////////////////////////////////////////////////
 
-void STableTreeView::BuildGroupingSubMenu_Change(FMenuBuilder& MenuBuilder, const TSharedPtr<FTreeNodeGrouping> CrumbGrouping)
+TSharedRef<SWidget> STableTreeView::CreateGroupingMenuWidget(const FTreeNodeGrouping& Grouping) const
 {
-	MenuBuilder.BeginSection("ChangeGrouping");
+	return SNew(SHorizontalBox)
+
+		+ SHorizontalBox::Slot()
+		.Padding(FMargin(3.0f, 0.0f, 8.0f, 0.0f))
+		.AutoWidth()
+		.VAlign(VAlign_Center)
+		[
+			SNew(SImage)
+			.Image(Grouping.GetIcon() ? Grouping.GetIcon() : FBaseTreeNode::GetDefaultIcon(true))
+			.DesiredSizeOverride(FVector2D(12.0, 12.0))
+			.ColorAndOpacity(Grouping.GetColor())
+		]
+
+		+ SHorizontalBox::Slot()
+		.Padding(FMargin(0.0f, 0.0f, 12.0f, 0.0f))
+		.FillWidth(1.0f)
+		.VAlign(VAlign_Center)
+		[
+			SNew(STextBlock)
+			.Text(Grouping.GetTitleName())
+		];
+}
+
+////////////////////////////////////////////////////////////////////////////////////////////////////
+
+void STableTreeView::AddGroupingMenuEntries(FMenuBuilder& MenuBuilder, EGroupingMenuType MenuType, const TSharedPtr<FTreeNodeGrouping> CrumbGrouping, TFunctionRef<bool(const FTreeNodeGrouping&)> FilterFunc)
+{
+	for (const TSharedPtr<FTreeNodeGrouping>& Grouping : AvailableGroupings)
 	{
-		for (const TSharedPtr<FTreeNodeGrouping>& Grouping : AvailableGroupings)
+		if (!FilterFunc(*Grouping))
 		{
-			FUIAction Action_Change
-			(
-				FExecuteAction::CreateSP(this, &STableTreeView::GroupingCrumbMenu_Change_Execute, CrumbGrouping, Grouping),
-				FCanExecuteAction::CreateSP(this, &STableTreeView::GroupingCrumbMenu_Change_CanExecute, CrumbGrouping, Grouping)
-			);
-			MenuBuilder.AddMenuEntry
-			(
-				Grouping->GetTitleName(),
-				Grouping->GetDescription(),
-				FSlateIcon(),
-				Action_Change,
-				NAME_None,
-				EUserInterfaceActionType::Button
-			);
+			continue;
 		}
+
+		FUIAction Action;
+		if (MenuType == EGroupingMenuType::Change)
+		{
+			Action.ExecuteAction = FExecuteAction::CreateSP(this, &STableTreeView::GroupingCrumbMenu_Change_Execute, CrumbGrouping, Grouping);
+			Action.CanExecuteAction = FCanExecuteAction::CreateSP(this, &STableTreeView::GroupingCrumbMenu_Change_CanExecute, CrumbGrouping, Grouping);
+		}
+		else // if (MenuType == EGroupingMenuType::Add)
+		{
+			Action.ExecuteAction = FExecuteAction::CreateSP(this, &STableTreeView::GroupingCrumbMenu_Add_Execute, Grouping, CrumbGrouping);
+			Action.CanExecuteAction = FCanExecuteAction::CreateSP(this, &STableTreeView::GroupingCrumbMenu_Add_CanExecute, Grouping, CrumbGrouping);
+		}
+
+		MenuBuilder.AddMenuEntry
+		(
+			Action,
+			CreateGroupingMenuWidget(*Grouping),
+			NAME_None,
+			Grouping->GetDescription(),
+			EUserInterfaceActionType::Button
+		);
 	}
-	MenuBuilder.EndSection();
 }
 
 ////////////////////////////////////////////////////////////////////////////////////////////////////
 
-void STableTreeView::BuildGroupingSubMenu_Add(FMenuBuilder& MenuBuilder, const TSharedPtr<FTreeNodeGrouping> CrumbGrouping)
+void STableTreeView::BuildGroupingSubMenu(FMenuBuilder& MenuBuilder, EGroupingMenuType MenuType, const TSharedPtr<FTreeNodeGrouping> CrumbGrouping)
 {
-	MenuBuilder.BeginSection("AddGrouping");
+	int32 NumUniqueValuesGroupings = 0;
+	int32 NumPathBreakdownGroupings = 0;
+	for (const TSharedPtr<FTreeNodeGrouping>& Grouping : AvailableGroupings)
 	{
-		for (const TSharedPtr<FTreeNodeGrouping>& Grouping : AvailableGroupings)
+		if (Grouping->Is<FTreeNodeGroupingByUniqueValue>())
 		{
-			FUIAction Action_Add
+			++NumUniqueValuesGroupings;
+		}
+		if (Grouping->Is<FTreeNodeGroupingByPathBreakdown>())
+		{
+			++NumPathBreakdownGroupings;
+		}
+	}
+	const bool bSubMenuForUniqueValuesGroupings = (NumUniqueValuesGroupings > 3);
+	const bool bSubMenuForPathBreakdownGroupings = (NumPathBreakdownGroupings > 3);
+
+	MenuBuilder.BeginSection(MenuType == EGroupingMenuType::Change ? "ChangeGrouping" : "AddGrouping");
+
+	if (bSubMenuForUniqueValuesGroupings || bSubMenuForPathBreakdownGroupings)
+	{
+		auto Filter = [bSubMenuForUniqueValuesGroupings, bSubMenuForPathBreakdownGroupings](const FTreeNodeGrouping& Grouping)
+			{
+				return (!(bSubMenuForUniqueValuesGroupings && Grouping.Is<FTreeNodeGroupingByUniqueValue>()) &&
+						!(bSubMenuForPathBreakdownGroupings && Grouping.Is<FTreeNodeGroupingByPathBreakdown>()));
+			};
+		AddGroupingMenuEntries(MenuBuilder, MenuType, CrumbGrouping, Filter);
+
+		if (bSubMenuForUniqueValuesGroupings)
+		{
+			TFunctionRef<bool(const FTreeNodeGrouping&)> FilterByUniqueValue = [](const FTreeNodeGrouping& Grouping)
+				{
+					return Grouping.Is<FTreeNodeGroupingByUniqueValue>();
+				};
+			MenuBuilder.AddSubMenu
 			(
-				FExecuteAction::CreateSP(this, &STableTreeView::GroupingCrumbMenu_Add_Execute, Grouping, CrumbGrouping),
-				FCanExecuteAction::CreateSP(this, &STableTreeView::GroupingCrumbMenu_Add_CanExecute, Grouping, CrumbGrouping)
+				LOCTEXT("UniqueValuesGroupings", "Unique Values"),
+				LOCTEXT("UniqueValuesGroupings_ToolTip", "Unique Values groupings"),
+				FNewMenuDelegate::CreateSP(this, &STableTreeView::AddGroupingMenuEntries, MenuType, CrumbGrouping, FilterByUniqueValue),
+				false,
+				FSlateIcon()
 			);
-			MenuBuilder.AddMenuEntry
+		}
+
+		if (bSubMenuForPathBreakdownGroupings)
+		{
+			TFunctionRef<bool(const FTreeNodeGrouping&)> FilterByPathBreakdown = [](const FTreeNodeGrouping& Grouping)
+				{
+					return Grouping.Is<FTreeNodeGroupingByPathBreakdown>();
+				};
+			MenuBuilder.AddSubMenu
 			(
-				Grouping->GetTitleName(),
-				Grouping->GetDescription(),
-				FSlateIcon(),
-				Action_Add,
-				NAME_None,
-				EUserInterfaceActionType::Button
+				LOCTEXT("PathBreakdownGroupings", "Path Breakdown"),
+				LOCTEXT("PathBreakdownGroupings_ToolTip", "Path Breakdown groupings"),
+				FNewMenuDelegate::CreateSP(this, &STableTreeView::AddGroupingMenuEntries, MenuType, CrumbGrouping, FilterByPathBreakdown),
+				false,
+				FSlateIcon()
 			);
 		}
 	}
+	else
+	{
+		auto Filter = [](const FTreeNodeGrouping& Grouping) { return true; };
+		AddGroupingMenuEntries(MenuBuilder, MenuType, CrumbGrouping, Filter);
+	}
+
 	MenuBuilder.EndSection();
 }
 
 ////////////////////////////////////////////////////////////////////////////////////////////////////
 
+TSharedRef<SWidget> STableTreeView::GetGroupingCrumbButtonContent(const TSharedPtr<FTreeNodeGrouping>& CrumbGrouping, const FTextBlockStyle* InTextStyle)
+{
+	return SNew(STextBlock)
+		.Text(CrumbGrouping->GetShortName())
+		.ColorAndOpacity(CrumbGrouping->GetColor());
+}
+
+////////////////////////////////////////////////////////////////////////////////////////////////////
+
 TSharedRef<SWidget> STableTreeView::GetGroupingCrumbMenuContent(const TSharedPtr<FTreeNodeGrouping>& CrumbGrouping)
 {
 	const bool bShouldCloseWindowAfterMenuSelection = true;
@@ -2263,7 +2352,7 @@ TSharedRef<SWidget> STableTreeView::GetGroupingCrumbMenuContent(const TSharedPtr
 		(
 			AddGroupingText,
 			LOCTEXT("GroupingMenu_AddOrInsert_SubMenu_Desc", "Adds or inserts a new grouping."),
-			FNewMenuDelegate::CreateSP(this, &STableTreeView::BuildGroupingSubMenu_Add, CrumbGrouping),
+			FNewMenuDelegate::CreateSP(this, &STableTreeView::BuildGroupingSubMenu, EGroupingMenuType::Add, CrumbGrouping),
 			false,
 			FSlateIcon()
 		);
@@ -2283,7 +2372,7 @@ TSharedRef<SWidget> STableTreeView::GetGroupingCrumbMenuContent(const TSharedPtr
 			(
 				LOCTEXT("GroupingMenu_Change_SubMenu", "Change To..."),
 				LOCTEXT("GroupingMenu_Change_SubMenu_Desc", "Changes the selected grouping."),
-				FNewMenuDelegate::CreateSP(this, &STableTreeView::BuildGroupingSubMenu_Change, CrumbGrouping),
+				FNewMenuDelegate::CreateSP(this, &STableTreeView::BuildGroupingSubMenu, EGroupingMenuType::Change, CrumbGrouping),
 				false,
 				FSlateIcon()
 			);
@@ -2471,7 +2560,7 @@ bool STableTreeView::GroupingCrumbMenu_Change_CanExecute(const TSharedPtr<FTreeN
 
 ////////////////////////////////////////////////////////////////////////////////////////////////////
 
-void STableTreeView::GroupingCrumbMenu_Add_Execute(const TSharedPtr<FTreeNodeGrouping> Grouping, const TSharedPtr<FTreeNodeGrouping> AfterGrouping)
+void STableTreeView::GroupingCrumbMenu_Add_Execute(const TSharedPtr<FTreeNodeGrouping> NewGrouping, const TSharedPtr<FTreeNodeGrouping> AfterGrouping)
 {
 	PreChangeGroupings();
 
@@ -2480,7 +2569,7 @@ void STableTreeView::GroupingCrumbMenu_Add_Execute(const TSharedPtr<FTreeNodeGro
 		const int32 AfterGroupingDepth = GetGroupingDepth(AfterGrouping);
 		ensure(AfterGroupingDepth >= 0);
 
-		const int32 GroupingDepth = GetGroupingDepth(Grouping);
+		const int32 GroupingDepth = GetGroupingDepth(NewGrouping);
 
 		if (GroupingDepth >= 0) // Grouping already exists
 		{
@@ -2488,22 +2577,22 @@ void STableTreeView::GroupingCrumbMenu_Add_Execute(const TSharedPtr<FTreeNodeGro
 
 			if (GroupingDepth <= AfterGroupingDepth)
 			{
-				CurrentGroupings.Insert(Grouping, AfterGroupingDepth);
+				CurrentGroupings.Insert(NewGrouping, AfterGroupingDepth);
 			}
 			else
 			{
-				CurrentGroupings.Insert(Grouping, AfterGroupingDepth + 1);
+				CurrentGroupings.Insert(NewGrouping, AfterGroupingDepth + 1);
 			}
 		}
 		else
 		{
-			CurrentGroupings.Insert(Grouping, AfterGroupingDepth + 1);
+			CurrentGroupings.Insert(NewGrouping, AfterGroupingDepth + 1);
 		}
 	}
 	else
 	{
-		CurrentGroupings.Remove(Grouping);
-		CurrentGroupings.Add(Grouping);
+		CurrentGroupings.Remove(NewGrouping);
+		CurrentGroupings.Add(NewGrouping);
 	}
 
 	PostChangeGroupings();
@@ -2511,20 +2600,20 @@ void STableTreeView::GroupingCrumbMenu_Add_Execute(const TSharedPtr<FTreeNodeGro
 
 ////////////////////////////////////////////////////////////////////////////////////////////////////
 
-bool STableTreeView::GroupingCrumbMenu_Add_CanExecute(const TSharedPtr<FTreeNodeGrouping> Grouping, const TSharedPtr<FTreeNodeGrouping> AfterGrouping) const
+bool STableTreeView::GroupingCrumbMenu_Add_CanExecute(const TSharedPtr<FTreeNodeGrouping> NewGrouping, const TSharedPtr<FTreeNodeGrouping> AfterGrouping) const
 {
 	if (AfterGrouping.IsValid())
 	{
 		const int32 AfterGroupingDepth = GetGroupingDepth(AfterGrouping);
 		ensure(AfterGroupingDepth >= 0);
 
-		const int32 GroupingDepth = GetGroupingDepth(Grouping);
+		const int32 GroupingDepth = GetGroupingDepth(NewGrouping);
 
 		return GroupingDepth < AfterGroupingDepth || GroupingDepth > AfterGroupingDepth + 1;
 	}
 	else
 	{
-		return Grouping != CurrentGroupings.Last();
+		return NewGrouping != CurrentGroupings.Last();
 	}
 }
 
diff --git a/Engine/Source/Developer/TraceInsightsCore/Public/InsightsCore/Table/ViewModels/TableTreeNode.h b/Engine/Source/Developer/TraceInsightsCore/Public/InsightsCore/Table/ViewModels/TableTreeNode.h
--- a/Engine/Source/Developer/TraceInsightsCore/Public/InsightsCore/Table/ViewModels/TableTreeNode.h
+++ b/Engine/Source/Developer/TraceInsightsCore/Public/InsightsCore/Table/ViewModels/TableTreeNode.h
@@ -190,6 +190,15 @@ public:
 	{
 	}
 
+	/** Initialization constructor for the group node. */
+	explicit FCustomTableTreeNode(const FName InName, TWeakPtr<FTable> InParentTable, FLinearColor InColor)
+		: FTableTreeNode(InName, InParentTable)
+		, IconBrush(FBaseTreeNode::GetDefaultIcon(true))
+		, IconColor(InColor)
+		, Color(InColor)
+	{
+	}
+
 	/** Initialization constructor for the group node. */
 	explicit FCustomTableTreeNode(const FName InName, TWeakPtr<FTable> InParentTable, const FSlateBrush* InIconBrush, FLinearColor InColor)
 		: FTableTreeNode(InName, InParentTable)
diff --git a/Engine/Source/Developer/TraceInsightsCore/Public/InsightsCore/Table/ViewModels/TreeNodeGrouping.h b/Engine/Source/Developer/TraceInsightsCore/Public/InsightsCore/Table/ViewModels/TreeNodeGrouping.h
--- a/Engine/Source/Developer/TraceInsightsCore/Public/InsightsCore/Table/ViewModels/TreeNodeGrouping.h
+++ b/Engine/Source/Developer/TraceInsightsCore/Public/InsightsCore/Table/ViewModels/TreeNodeGrouping.h
@@ -38,20 +38,61 @@ public:
 	virtual FText GetTitleName() const = 0;
 	virtual FText GetDescription() const = 0;
 
-	UE_DEPRECATED(5.6, "GetBrushName() is not used")
-	virtual FName GetBrushName() const { return NAME_None; }
-
 	virtual const FSlateBrush* GetIcon() const = 0;
+	virtual const FLinearColor& GetColor() const = 0;
 
 	virtual FName GetColumnId() const = 0;
 };
 
 ////////////////////////////////////////////////////////////////////////////////////////////////////
 
-struct TRACEINSIGHTSCORE_API FTreeNodeGroupInfo
+class TRACEINSIGHTSCORE_API FTreeNodeGroupInfo
 {
+public:
+	FTreeNodeGroupInfo() = default;
+
+	FTreeNodeGroupInfo(const FName& InName, bool bInIsExpanded = false)
+		: Name(InName)
+		, bIsExpanded(bInIsExpanded)
+	{
+	}
+
+	FTreeNodeGroupInfo(const FName& InName, bool bInIsExpanded, const FLinearColor& InColor)
+		: Name(InName)
+		, bIsExpanded(bInIsExpanded)
+		, bIsCustom(true)
+		, Icon(nullptr)
+		, IconColor(InColor)
+		, Color(InColor)
+	{
+	}
+
+	FTreeNodeGroupInfo(const FName& InName, bool bInIsExpanded, const FSlateBrush* InIcon, const FLinearColor& InIconColor, const FLinearColor& InColor)
+		: Name(InName)
+		, bIsExpanded(bInIsExpanded)
+		, bIsCustom(true)
+		, Icon(InIcon)
+		, IconColor(InIconColor)
+		, Color(InColor)
+	{
+	}
+
+	const FName& GetName() const { return Name; }
+
+	bool IsExpanded() const { return bIsExpanded; }
+	bool IsCustom() const { return bIsCustom; }
+
+	const FSlateBrush* GetIcon() const { return Icon; }
+	const FLinearColor& GetIconColor() const { return IconColor; }
+	const FLinearColor& GetColor() const { return Color; }
+
+private:
 	FName Name;
-	bool IsExpanded;
+	bool bIsExpanded = false;
+	bool bIsCustom = false;
+	const FSlateBrush* Icon = nullptr;
+	FLinearColor IconColor;
+	FLinearColor Color;
 };
 
 ////////////////////////////////////////////////////////////////////////////////////////////////////
@@ -61,38 +102,36 @@ class TRACEINSIGHTSCORE_API FTreeNodeGrouping : public ITreeNodeGrouping
 	INSIGHTS_DECLARE_RTTI(FTreeNodeGrouping, ITreeNodeGrouping)
 
 public:
-	FTreeNodeGrouping() : Icon(nullptr) {}
+	FTreeNodeGrouping();
 	FTreeNodeGrouping(const FText& InShortName, const FText& InTitleName, const FText& InDescription, const FSlateBrush* InIcon);
 
-	UE_DEPRECATED(5.6, "BrushName is not used")
-	FTreeNodeGrouping(const FText& InShortName, const FText& InTitleName, const FText& InDescription, const FName InBrushName, const FSlateBrush* InIcon);
-
 	FTreeNodeGrouping(const FTreeNodeGrouping&) = delete;
 	FTreeNodeGrouping& operator=(FTreeNodeGrouping&) = delete;
 
 	virtual ~FTreeNodeGrouping() {}
 
+	// ITreeNodeGrouping
 	virtual FText GetShortName() const override { return ShortName; }
 	virtual FText GetTitleName() const override { return TitleName; }
 	virtual FText GetDescription() const override { return Description; }
-
-	UE_DEPRECATED(5.6, "GetBrushName() is not used")
-	virtual FName GetBrushName() const override;
-
 	virtual const FSlateBrush* GetIcon() const override { return Icon; }
-
+	virtual const FLinearColor& GetColor() const override { return Color; }
 	virtual FName GetColumnId() const override { return NAME_None; }
 
-	virtual FTreeNodeGroupInfo GetGroupForNode(const FBaseTreeNodePtr InNode) const { return { FName(), false }; }
+	virtual void SetIcon(const FSlateBrush* InIcon) { Icon = InIcon; }
+	virtual void SetColor(const FLinearColor& InColor) { Color = InColor; }
+
+	virtual FTreeNodeGroupInfo GetGroupForNode(const FBaseTreeNodePtr InNode) const { return FTreeNodeGroupInfo(); }
 	virtual void GroupNodes(const TArray<FTableTreeNodePtr>& Nodes, FTableTreeNode& ParentGroup, TWeakPtr<FTable> InParentTable, IAsyncOperationProgress& InAsyncOperationProgress) const;
 
+	static FLinearColor MakeDistinctColor(const FName& InName);
+
 protected:
 	FText ShortName;
 	FText TitleName;
 	FText Description;
-	UE_DEPRECATED(5.6, "BrushName is not used")
-	FName BrushName;
-	const FSlateBrush* Icon;
+	const FSlateBrush* Icon = nullptr;
+	FLinearColor Color;
 };
 
 ////////////////////////////////////////////////////////////////////////////////////////////////////
@@ -198,7 +237,7 @@ void TTreeNodeGroupingByUniqueValue<Type>::GroupNodes(const TArray<FTableTreeNod
 			if (!GroupPtrPtr)
 			{
 				const FName GroupName = GetGroupName(Column, *NodePtr);
-				GroupPtr = MakeShared<FTableTreeNode>(GroupName, InParentTable);
+				GroupPtr = MakeShared<FCustomTableTreeNode>(GroupName, InParentTable, GetColor());
 				GroupPtr->SetExpansion(false);
 				ParentGroup.AddChildAndSetParent(GroupPtr);
 				GroupMap.Add(Value, GroupPtr);
diff --git a/Engine/Source/Developer/TraceInsightsCore/Public/InsightsCore/Table/Widgets/STableTreeView.h b/Engine/Source/Developer/TraceInsightsCore/Public/InsightsCore/Table/Widgets/STableTreeView.h
--- a/Engine/Source/Developer/TraceInsightsCore/Public/InsightsCore/Table/Widgets/STableTreeView.h
+++ b/Engine/Source/Developer/TraceInsightsCore/Public/InsightsCore/Table/Widgets/STableTreeView.h
@@ -114,6 +114,13 @@ class TRACEINSIGHTSCORE_API STableTreeView : public SCompoundWidget, public IAsy
 	friend class FSearchForItemToSelectTask;
 	friend class FSelectNodeByTableRowIndexTask;
 
+protected:
+	enum class EGroupingMenuType
+	{
+		Change,
+		Add
+	};
+
 public:
 	/** Default constructor. */
 	STableTreeView();
@@ -334,9 +341,11 @@ protected:
 	void GroupNodesRec(const TArray<FTableTreeNodePtr>& Nodes, FTableTreeNode& ParentGroup, int32 GroupingDepth, const TArray<TSharedPtr<FTreeNodeGrouping>>& Groupings);
 
 	void RebuildGroupingCrumbs();
+	TSharedRef<SWidget> GetGroupingCrumbButtonContent(const TSharedPtr<FTreeNodeGrouping>& CrumbGrouping, const FTextBlockStyle* InTextStyle);
 	void OnGroupingCrumbClicked(const TSharedPtr<FTreeNodeGrouping>& InEntry);
-	void BuildGroupingSubMenu_Change(FMenuBuilder& MenuBuilder, const TSharedPtr<FTreeNodeGrouping> CrumbGrouping);
-	void BuildGroupingSubMenu_Add(FMenuBuilder& MenuBuilder, const TSharedPtr<FTreeNodeGrouping> CrumbGrouping);
+	TSharedRef<SWidget> CreateGroupingMenuWidget(const FTreeNodeGrouping& Grouping) const;
+	void AddGroupingMenuEntries(FMenuBuilder& MenuBuilder, EGroupingMenuType MenuType, const TSharedPtr<FTreeNodeGrouping> CrumbGrouping, TFunctionRef<bool(const FTreeNodeGrouping&)> FilterFunc);
+	void BuildGroupingSubMenu(FMenuBuilder& MenuBuilder, EGroupingMenuType MenuType, const TSharedPtr<FTreeNodeGrouping> CrumbGrouping);
 	TSharedRef<SWidget> GetGroupingCrumbMenuContent(const TSharedPtr<FTreeNodeGrouping>& CrumbGrouping);
 
 	void PreChangeGroupings();
