Fixed CIS error (re TFunction/TFunctionRef scope).

#rnx

diff --git a/Engine/Source/Developer/TraceInsightsCore/Private/InsightsCore/Table/Widgets/STableTreeView.cpp b/Engine/Source/Developer/TraceInsightsCore/Private/InsightsCore/Table/Widgets/STableTreeView.cpp
--- a/Engine/Source/Developer/TraceInsightsCore/Private/InsightsCore/Table/Widgets/STableTreeView.cpp
+++ b/Engine/Source/Developer/TraceInsightsCore/Private/InsightsCore/Table/Widgets/STableTreeView.cpp
@@ -2221,7 +2221,7 @@ TSharedRef<SWidget> STableTreeView::CreateGroupingMenuWidget(const FTreeNodeGrou
 
 ////////////////////////////////////////////////////////////////////////////////////////////////////
 
-void STableTreeView::AddGroupingMenuEntries(FMenuBuilder& MenuBuilder, EGroupingMenuType MenuType, const TSharedPtr<FTreeNodeGrouping> CrumbGrouping, TFunctionRef<bool(const FTreeNodeGrouping&)> FilterFunc)
+void STableTreeView::AddGroupingMenuEntries(FMenuBuilder& MenuBuilder, EGroupingMenuType MenuType, const TSharedPtr<FTreeNodeGrouping> CrumbGrouping, TFunction<bool(const FTreeNodeGrouping&)> FilterFunc)
 {
 	for (const TSharedPtr<FTreeNodeGrouping>& Grouping : AvailableGroupings)
 	{
@@ -2286,7 +2286,7 @@ void STableTreeView::BuildGroupingSubMenu(FMenuBuilder& MenuBuilder, EGroupingMe
 
 		if (bSubMenuForUniqueValuesGroupings)
 		{
-			TFunctionRef<bool(const FTreeNodeGrouping&)> FilterByUniqueValue = [](const FTreeNodeGrouping& Grouping)
+			TFunction<bool(const FTreeNodeGrouping&)> FilterByUniqueValue = [](const FTreeNodeGrouping& Grouping)
 				{
 					return Grouping.Is<FTreeNodeGroupingByUniqueValue>();
 				};
@@ -2302,7 +2302,7 @@ void STableTreeView::BuildGroupingSubMenu(FMenuBuilder& MenuBuilder, EGroupingMe
 
 		if (bSubMenuForPathBreakdownGroupings)
 		{
-			TFunctionRef<bool(const FTreeNodeGrouping&)> FilterByPathBreakdown = [](const FTreeNodeGrouping& Grouping)
+			TFunction<bool(const FTreeNodeGrouping&)> FilterByPathBreakdown = [](const FTreeNodeGrouping& Grouping)
 				{
 					return Grouping.Is<FTreeNodeGroupingByPathBreakdown>();
 				};
diff --git a/Engine/Source/Developer/TraceInsightsCore/Public/InsightsCore/Table/Widgets/STableTreeView.h b/Engine/Source/Developer/TraceInsightsCore/Public/InsightsCore/Table/Widgets/STableTreeView.h
--- a/Engine/Source/Developer/TraceInsightsCore/Public/InsightsCore/Table/Widgets/STableTreeView.h
+++ b/Engine/Source/Developer/TraceInsightsCore/Public/InsightsCore/Table/Widgets/STableTreeView.h
@@ -344,7 +344,7 @@ protected:
 	TSharedRef<SWidget> GetGroupingCrumbButtonContent(const TSharedPtr<FTreeNodeGrouping>& CrumbGrouping, const FTextBlockStyle* InTextStyle);
 	void OnGroupingCrumbClicked(const TSharedPtr<FTreeNodeGrouping>& InEntry);
 	TSharedRef<SWidget> CreateGroupingMenuWidget(const FTreeNodeGrouping& Grouping) const;
-	void AddGroupingMenuEntries(FMenuBuilder& MenuBuilder, EGroupingMenuType MenuType, const TSharedPtr<FTreeNodeGrouping> CrumbGrouping, TFunctionRef<bool(const FTreeNodeGrouping&)> FilterFunc);
+	void AddGroupingMenuEntries(FMenuBuilder& MenuBuilder, EGroupingMenuType MenuType, const TSharedPtr<FTreeNodeGrouping> CrumbGrouping, TFunction<bool(const FTreeNodeGrouping&)> FilterFunc);
 	void BuildGroupingSubMenu(FMenuBuilder& MenuBuilder, EGroupingMenuType MenuType, const TSharedPtr<FTreeNodeGrouping> CrumbGrouping);
 	TSharedRef<SWidget> GetGroupingCrumbMenuContent(const TSharedPtr<FTreeNodeGrouping>& CrumbGrouping);
 
