[InsightsCore]
* Made it possible to provide x_API to INSIGHTS_DECLARE_RTTI

diff --git a/Engine/Source/Developer/TraceInsightsCore/Public/InsightsCore/Common/SimpleRtti.h b/Engine/Source/Developer/TraceInsightsCore/Public/InsightsCore/Common/SimpleRtti.h
--- a/Engine/Source/Developer/TraceInsightsCore/Public/InsightsCore/Common/SimpleRtti.h
+++ b/Engine/Source/Developer/TraceInsightsCore/Public/InsightsCore/Common/SimpleRtti.h
@@ -8,7 +8,7 @@
 //       Ex.: Some classes in the actual hierarchy can be hidden (do not use INSIGHTS_DECLARE_RTTI, so just inherits type info from parent class).
 
 // Macro to be used in declaration of the base class of the SimpleRTTI class hierarchy.
-#define INSIGHTS_DECLARE_RTTI_BASE(BaseClassName) \
+#define INSIGHTS_DECLARE_RTTI_BASE(BaseClassName, ...) \
 	public: \
 		static const FName& GetStaticTypeName() { return BaseClassName::TypeName; } \
 		virtual const FName& GetTypeName() const { return BaseClassName::GetStaticTypeName(); } \
@@ -17,16 +17,16 @@
 		template<typename Type> const Type& As() const { return *static_cast<const Type*>(this); } \
 		template<typename Type> Type& As() { return *static_cast<Type*>(this); } \
 	private: \
-		static const FName TypeName;
+		static __VA_ARGS__ const FName TypeName;
 
 // Macro to be used in declaration of a derived class in a SimpleRTTI hierarchy.
-#define INSIGHTS_DECLARE_RTTI(ClassName, BaseClassName) \
+#define INSIGHTS_DECLARE_RTTI(ClassName, BaseClassName, ...) \
 	public: \
 		static const FName& GetStaticTypeName() { return ClassName::TypeName; } \
 		virtual const FName& GetTypeName() const override { return ClassName::GetStaticTypeName(); } \
 		virtual bool IsKindOf(const FName& InTypeName) const override { return InTypeName == ClassName::GetStaticTypeName() || BaseClassName::IsKindOf(InTypeName); } \
 	private: \
-		static const FName TypeName;
+		static __VA_ARGS__ const FName TypeName;
 
 // Macro to be used in implementation of a class in a SimpleRTTI hierarchy.
 #define INSIGHTS_IMPLEMENT_RTTI(ClassName) \
