Updated headers using UnrealCodeFixup to make sure dllstorage is on methods/staticvars instead of types. Part 4/n

diff --git a/Engine/Source/Developer/TraceInsightsCore/Public/InsightsCore/Filter/ViewModels/Filters.h b/Engine/Source/Developer/TraceInsightsCore/Public/InsightsCore/Filter/ViewModels/Filters.h
--- a/Engine/Source/Developer/TraceInsightsCore/Public/InsightsCore/Filter/ViewModels/Filters.h
+++ b/Engine/Source/Developer/TraceInsightsCore/Public/InsightsCore/Filter/ViewModels/Filters.h
@@ -13,6 +13,8 @@
 
 #include "InsightsCore/Common/SimpleRtti.h"
 
+#define UE_API TRACEINSIGHTSCORE_API
+
 class FSpawnTabArgs;
 class SDockTab;
 class SHorizontalBox;
@@ -197,9 +199,9 @@ typedef TSharedPtr<TArray<TSharedPtr<IFilterOperator>>> SupportedOperatorsArrayP
 
 class FFilter;
 
-class TRACEINSIGHTSCORE_API FFilterState
+class FFilterState
 {
-	INSIGHTS_DECLARE_RTTI_BASE(FFilterState)
+	INSIGHTS_DECLARE_RTTI_BASE(FFilterState, UE_API)
 
 public:
 	FFilterState(TSharedRef<FFilter> InFilter)
@@ -213,15 +215,15 @@ public:
 	TSharedPtr<const IFilterOperator> GetSelectedOperator() const { return SelectedOperator; }
 
 	virtual void Update() {};
-	virtual bool ApplyFilter(const FFilterContext& Context) const;
+	UE_API virtual bool ApplyFilter(const FFilterContext& Context) const;
 
-	virtual void SetFilterValue(FString InTextValue);
+	UE_API virtual void SetFilterValue(FString InTextValue);
 
 	virtual bool HasCustomUI() const { return false; }
 	virtual void AddCustomUI(TSharedRef<SHorizontalBox> LeftBox) {}
 
-	virtual bool Equals(const FFilterState& Other) const;
-	virtual TSharedRef<FFilterState> DeepCopy() const;
+	UE_API virtual bool Equals(const FFilterState& Other) const;
+	UE_API virtual TSharedRef<FFilterState> DeepCopy() const;
 
 protected:
 	TSharedRef<FFilter> Filter;
@@ -231,9 +233,9 @@ protected:
 
 ////////////////////////////////////////////////////////////////////////////////////////////////////
 
-class TRACEINSIGHTSCORE_API FFilter : public TSharedFromThis<FFilter>
+class FFilter : public TSharedFromThis<FFilter>
 {
-	INSIGHTS_DECLARE_RTTI_BASE(FFilter)
+	INSIGHTS_DECLARE_RTTI_BASE(FFilter, UE_API)
 
 public:
 	FFilter(int32 InKey, FText InName, FText InDesc, EFilterDataType InDataType, TSharedPtr<IFilterValueConverter> InConverter, SupportedOperatorsArrayPtr InSupportedOperators)
@@ -271,19 +273,19 @@ protected:
 
 ////////////////////////////////////////////////////////////////////////////////////////////////////
 
-class TRACEINSIGHTSCORE_API FFilterWithSuggestionsValueConverter : public IFilterValueConverter
+class FFilterWithSuggestionsValueConverter : public IFilterValueConverter
 {
 public:
 	virtual bool Convert(const FString& Input, double& Output, FText& OutError) const override { return true; }
-	virtual FText GetTooltipText() const override;
-	virtual FText GetHintText() const override;
+	UE_API virtual FText GetTooltipText() const override;
+	UE_API virtual FText GetHintText() const override;
 };
 
 ////////////////////////////////////////////////////////////////////////////////////////////////////
 
-class TRACEINSIGHTSCORE_API FFilterWithSuggestions : public FFilter
+class FFilterWithSuggestions : public FFilter
 {
-	INSIGHTS_DECLARE_RTTI(FFilterWithSuggestions, FFilter)
+	INSIGHTS_DECLARE_RTTI(FFilterWithSuggestions, FFilter, UE_API)
 
 public:
 	typedef TFunction<void(const FString& /*Text*/, TArray<FString>& OutSuggestions)> FGetSuggestionsCallback;
@@ -313,11 +315,11 @@ private:
 
 ////////////////////////////////////////////////////////////////////////////////////////////////////
 
-class TRACEINSIGHTSCORE_API FFilterStorage
+class FFilterStorage
 {
 public:
-	FFilterStorage();
-	~FFilterStorage();
+	UE_API FFilterStorage();
+	UE_API ~FFilterStorage();
 
 	const TArray<TSharedPtr<FFilterGroupOperator>>& GetFilterGroupOperators() const { return FilterGroupOperators; }
 
@@ -335,19 +337,19 @@ private:
 
 ////////////////////////////////////////////////////////////////////////////////////////////////////
 
-class TRACEINSIGHTSCORE_API FFilterService
+class FFilterService
 {
 public:
-	FFilterService();
-	~FFilterService();
+	UE_API FFilterService();
+	UE_API ~FFilterService();
 
-	static void Initialize();
-	static void Shutdown();
+	static UE_API void Initialize();
+	static UE_API void Shutdown();
 	static TSharedPtr<FFilterService> Get() { return Instance; }
 
-	void RegisterTabSpawner();
-	void UnregisterTabSpawner();
-	TSharedRef<SDockTab> SpawnTab(const FSpawnTabArgs& Args);
+	UE_API void RegisterTabSpawner();
+	UE_API void UnregisterTabSpawner();
+	UE_API TSharedRef<SDockTab> SpawnTab(const FSpawnTabArgs& Args);
 
 	const TArray<TSharedPtr<FFilterGroupOperator>>& GetFilterGroupOperators() const { return FilterStorage.GetFilterGroupOperators(); }
 
@@ -355,11 +357,11 @@ public:
 	TSharedPtr<TArray<TSharedPtr<IFilterOperator>>> GetIntegerOperators() const { return FilterStorage.GetIntegerOperators(); }
 	TSharedPtr<TArray<TSharedPtr<IFilterOperator>>> GetStringOperators() const { return FilterStorage.GetStringOperators(); }
 
-	TSharedPtr<SWidget> CreateFilterConfiguratorWidget(TSharedPtr<class FFilterConfigurator> FilterConfiguratorViewModel, TWeakPtr<SWidget> ParentWidget = nullptr);
+	UE_API TSharedPtr<SWidget> CreateFilterConfiguratorWidget(TSharedPtr<class FFilterConfigurator> FilterConfiguratorViewModel, TWeakPtr<SWidget> ParentWidget = nullptr);
 
 private:
-	static const FName FilterConfiguratorTabId;
-	static TSharedPtr<FFilterService> Instance;
+	static UE_API const FName FilterConfiguratorTabId;
+	static UE_API TSharedPtr<FFilterService> Instance;
 	FFilterStorage FilterStorage;
 	TSharedPtr<class SAdvancedFilter> PendingWidget;
 	TWeakPtr<SWidget> PendingParentWidget;
@@ -368,3 +370,5 @@ private:
 ////////////////////////////////////////////////////////////////////////////////////////////////////
 
 } // namespace UE::Insights
+
+#undef UE_API
diff --git a/Engine/Source/Developer/TraceInsightsCore/Public/InsightsCore/Filter/ViewModels/IFilterExecutor.h b/Engine/Source/Developer/TraceInsightsCore/Public/InsightsCore/Filter/ViewModels/IFilterExecutor.h
--- a/Engine/Source/Developer/TraceInsightsCore/Public/InsightsCore/Filter/ViewModels/IFilterExecutor.h
+++ b/Engine/Source/Developer/TraceInsightsCore/Public/InsightsCore/Filter/ViewModels/IFilterExecutor.h
@@ -9,7 +9,7 @@ namespace UE::Insights
 
 class FFilterContext;
 
-class TRACEINSIGHTSCORE_API IFilterExecutor
+class IFilterExecutor
 {
 public:
 	virtual bool ApplyFilters(const class FFilterContext& Context) const = 0;
diff --git a/Engine/Source/Developer/TraceInsightsCore/Public/InsightsCore/Filter/ViewModels/TimeFilterValueConverter.h b/Engine/Source/Developer/TraceInsightsCore/Public/InsightsCore/Filter/ViewModels/TimeFilterValueConverter.h
--- a/Engine/Source/Developer/TraceInsightsCore/Public/InsightsCore/Filter/ViewModels/TimeFilterValueConverter.h
+++ b/Engine/Source/Developer/TraceInsightsCore/Public/InsightsCore/Filter/ViewModels/TimeFilterValueConverter.h
@@ -9,15 +9,19 @@
 
 #include "InsightsCore/Filter/ViewModels/Filters.h"
 
+#define UE_API TRACEINSIGHTSCORE_API
+
 namespace UE::Insights
 {
 
-class TRACEINSIGHTSCORE_API FTimeFilterValueConverter : public IFilterValueConverter
+class FTimeFilterValueConverter : public IFilterValueConverter
 {
 public:
-	virtual bool Convert(const FString& Input, double& Output, FText& OutError) const override;
-	virtual FText GetTooltipText() const override;
-	virtual FText GetHintText() const override;
+	UE_API virtual bool Convert(const FString& Input, double& Output, FText& OutError) const override;
+	UE_API virtual FText GetTooltipText() const override;
+	UE_API virtual FText GetHintText() const override;
 };
 
 } // namespace UE::Insights
+
+#undef UE_API
diff --git a/Engine/Source/Developer/TraceInsightsCore/Public/InsightsCore/Filter/Widgets/SAdvancedFilter.h b/Engine/Source/Developer/TraceInsightsCore/Public/InsightsCore/Filter/Widgets/SAdvancedFilter.h
--- a/Engine/Source/Developer/TraceInsightsCore/Public/InsightsCore/Filter/Widgets/SAdvancedFilter.h
+++ b/Engine/Source/Developer/TraceInsightsCore/Public/InsightsCore/Filter/Widgets/SAdvancedFilter.h
@@ -11,6 +11,8 @@
 #include "InsightsCore/Filter/ViewModels/FilterConfiguratorNode.h"
 #include "InsightsCore/Filter/ViewModels/Filters.h"
 
+#define UE_API TRACEINSIGHTSCORE_API
+
 class SDockTab;
 
 namespace UE::Insights
@@ -23,14 +25,14 @@ class SFilterConfigurator;
 /**
  * A custom widget used to configure custom filters.
  */
-class TRACEINSIGHTSCORE_API SAdvancedFilter: public SCompoundWidget
+class SAdvancedFilter: public SCompoundWidget
 {
 public:
 	/** Default constructor. */
-	SAdvancedFilter();
+	UE_API SAdvancedFilter();
 
 	/** Virtual destructor. */
-	virtual ~SAdvancedFilter();
+	UE_API virtual ~SAdvancedFilter();
 
 	SLATE_BEGIN_ARGS(SAdvancedFilter) {}
 	SLATE_END_ARGS()
@@ -39,21 +41,21 @@ public:
 	 * Construct this widget
 	 * @param InArgs - The declaration data for this widget
 	 */
-	void Construct(const FArguments& InArgs, TSharedPtr<FFilterConfigurator> InFilterConfiguratorViewModel);
+	UE_API void Construct(const FArguments& InArgs, TSharedPtr<FFilterConfigurator> InFilterConfiguratorViewModel);
 
-	void Reset();
+	UE_API void Reset();
 
 	void SetParentTab(const TSharedPtr<SDockTab> InTab) { ParentTab = InTab; }
 	const TWeakPtr<SDockTab> GetParentTab() { return ParentTab; };
 
-	void RequestClose();
+	UE_API void RequestClose();
 
 private:
-	void InitCommandList();
+	UE_API void InitCommandList();
 
-	FReply OK_OnClicked();
+	UE_API FReply OK_OnClicked();
 
-	FReply Cancel_OnClicked();
+	UE_API FReply Cancel_OnClicked();
 
 private:
 	TSharedPtr<SFilterConfigurator> FilterConfigurator;
@@ -70,3 +72,5 @@ private:
 ////////////////////////////////////////////////////////////////////////////////////////////////////
 
 } // namespace UE::Insights
+
+#undef UE_API
diff --git a/Engine/Source/Developer/TraceInsightsCore/Public/InsightsCore/Filter/Widgets/SFilterConfigurator.h b/Engine/Source/Developer/TraceInsightsCore/Public/InsightsCore/Filter/Widgets/SFilterConfigurator.h
--- a/Engine/Source/Developer/TraceInsightsCore/Public/InsightsCore/Filter/Widgets/SFilterConfigurator.h
+++ b/Engine/Source/Developer/TraceInsightsCore/Public/InsightsCore/Filter/Widgets/SFilterConfigurator.h
@@ -11,6 +11,8 @@
 #include "InsightsCore/Filter/ViewModels/FilterConfiguratorNode.h"
 #include "InsightsCore/Filter/ViewModels/Filters.h"
 
+#define UE_API TRACEINSIGHTSCORE_API
+
 class SDockTab;
 
 namespace UE::Insights
@@ -23,14 +25,14 @@ class FFilterConfigurator;
 /**
  * A custom widget used to configure custom filters.
  */
-class TRACEINSIGHTSCORE_API SFilterConfigurator: public SCompoundWidget
+class SFilterConfigurator: public SCompoundWidget
 {
 public:
 	/** Default constructor. */
-	SFilterConfigurator();
+	UE_API SFilterConfigurator();
 
 	/** Virtual destructor. */
-	virtual ~SFilterConfigurator();
+	UE_API virtual ~SFilterConfigurator();
 
 	SLATE_BEGIN_ARGS(SFilterConfigurator) {}
 	SLATE_END_ARGS()
@@ -39,29 +41,29 @@ public:
 	 * Construct this widget
 	 * @param InArgs - The declaration data for this widget
 	 */
-	void Construct(const FArguments& InArgs, TSharedPtr<FFilterConfigurator> InFilterConfiguratorViewModel);
+	UE_API void Construct(const FArguments& InArgs, TSharedPtr<FFilterConfigurator> InFilterConfiguratorViewModel);
 
-	void Reset();
+	UE_API void Reset();
 
 	/**
 	 * Called by STreeView to retrieves the children for the specified parent item.
 	 * @param InParent    - The parent node to retrieve the children from.
 	 * @param OutChildren - List of children for the parent node.
 	 */
-	void TreeView_OnGetChildren(FFilterConfiguratorNodePtr InParent, TArray<FFilterConfiguratorNodePtr>& OutChildren);
+	UE_API void TreeView_OnGetChildren(FFilterConfiguratorNodePtr InParent, TArray<FFilterConfiguratorNodePtr>& OutChildren);
 
 	/** Called by STreeView to generate a table row for the specified item. */
-	TSharedRef<class ITableRow> TreeView_OnGenerateRow(FFilterConfiguratorNodePtr TreeNode, const TSharedRef<class STableViewBase>& OwnerTable);
+	UE_API TSharedRef<class ITableRow> TreeView_OnGenerateRow(FFilterConfiguratorNodePtr TreeNode, const TSharedRef<class STableViewBase>& OwnerTable);
 
 	void SetParentTab(const TSharedPtr<SDockTab> InTab) { ParentTab = InTab; }
 	const TWeakPtr<SDockTab> GetParentTab() { return ParentTab; };
 
-	void RequestClose();
+	UE_API void RequestClose();
 
 private:
-	void InitCommandList();
+	UE_API void InitCommandList();
 
-	void SetInitialExpansionRec(const FFilterConfiguratorNodePtr& Node, bool Value);
+	UE_API void SetInitialExpansionRec(const FFilterConfiguratorNodePtr& Node, bool Value);
 
 private:
 	TSharedPtr<FFilterConfigurator> FilterConfiguratorViewModel;
@@ -84,3 +86,5 @@ private:
 ////////////////////////////////////////////////////////////////////////////////////////////////////
 
 } // namespace UE::Insights
+
+#undef UE_API
diff --git a/Engine/Source/Developer/TraceInsightsCore/Public/InsightsCore/Filter/Widgets/SFilterConfiguratorRow.h b/Engine/Source/Developer/TraceInsightsCore/Public/InsightsCore/Filter/Widgets/SFilterConfiguratorRow.h
--- a/Engine/Source/Developer/TraceInsightsCore/Public/InsightsCore/Filter/Widgets/SFilterConfiguratorRow.h
+++ b/Engine/Source/Developer/TraceInsightsCore/Public/InsightsCore/Filter/Widgets/SFilterConfiguratorRow.h
@@ -9,6 +9,8 @@
 
 #include "InsightsCore/Filter/ViewModels/FilterConfiguratorNode.h"
 
+#define UE_API TRACEINSIGHTSCORE_API
+
 namespace UE::Insights
 {
 
@@ -17,7 +19,7 @@ class FFilterGroupOperator;
 class IFilterOperator;
 
 /** Widget that represents a table row in the Filter's tree control. Generates widgets for each column on demand. */
-class TRACEINSIGHTSCORE_API SFilterConfiguratorRow : public SMultiColumnTableRow<FFilterConfiguratorNodePtr>
+class SFilterConfiguratorRow : public SMultiColumnTableRow<FFilterConfiguratorNodePtr>
 {
 public:
 	SFilterConfiguratorRow() {}
@@ -28,54 +30,54 @@ public:
 	SLATE_END_ARGS()
 
 public:
-	void Construct(const FArguments& InArgs, const TSharedRef<STableViewBase>& InOwnerTableView);
+	UE_API void Construct(const FArguments& InArgs, const TSharedRef<STableViewBase>& InOwnerTableView);
 
-	virtual TSharedRef<SWidget> GenerateWidgetForColumn(const FName& InColumnName) override;
+	UE_API virtual TSharedRef<SWidget> GenerateWidgetForColumn(const FName& InColumnName) override;
 
 private:
-	const TArray<TSharedPtr<FFilter>>* GetAvailableFilters();
+	UE_API const TArray<TSharedPtr<FFilter>>* GetAvailableFilters();
 
-	TSharedRef<SWidget> AvailableFilters_OnGenerateWidget(TSharedPtr<FFilter> InFilter);
+	UE_API TSharedRef<SWidget> AvailableFilters_OnGenerateWidget(TSharedPtr<FFilter> InFilter);
 
-	void AvailableFilters_OnSelectionChanged(TSharedPtr<FFilter> InFilter, ESelectInfo::Type SelectInfo);
+	UE_API void AvailableFilters_OnSelectionChanged(TSharedPtr<FFilter> InFilter, ESelectInfo::Type SelectInfo);
 
-	FText AvailableFilters_GetSelectionText() const;
+	UE_API FText AvailableFilters_GetSelectionText() const;
 
-	const TArray<TSharedPtr<IFilterOperator>>* GetAvailableFilterOperators();
+	UE_API const TArray<TSharedPtr<IFilterOperator>>* GetAvailableFilterOperators();
 
-	TSharedRef<SWidget> AvailableFilterOperators_OnGenerateWidget(TSharedPtr<IFilterOperator> InFilter);
+	UE_API TSharedRef<SWidget> AvailableFilterOperators_OnGenerateWidget(TSharedPtr<IFilterOperator> InFilter);
 
-	void AvailableFilterOperators_OnSelectionChanged(TSharedPtr<IFilterOperator> InFilter, ESelectInfo::Type SelectInfo);
+	UE_API void AvailableFilterOperators_OnSelectionChanged(TSharedPtr<IFilterOperator> InFilter, ESelectInfo::Type SelectInfo);
 
-	FText AvailableFilterOperators_GetSelectionText() const;
+	UE_API FText AvailableFilterOperators_GetSelectionText() const;
 
-	const TArray<TSharedPtr<FFilterGroupOperator>>* GetFilterGroupOperators();
+	UE_API const TArray<TSharedPtr<FFilterGroupOperator>>* GetFilterGroupOperators();
 
-	TSharedRef<SWidget> FilterGroupOperators_OnGenerateWidget(TSharedPtr<FFilterGroupOperator> InFilter);
+	UE_API TSharedRef<SWidget> FilterGroupOperators_OnGenerateWidget(TSharedPtr<FFilterGroupOperator> InFilter);
 
-	void FilterGroupOperators_OnSelectionChanged(TSharedPtr<FFilterGroupOperator> InFilter, ESelectInfo::Type SelectInfo);
+	UE_API void FilterGroupOperators_OnSelectionChanged(TSharedPtr<FFilterGroupOperator> InFilter, ESelectInfo::Type SelectInfo);
 
-	FText FilterGroupOperators_GetSelectionText() const;
+	UE_API FText FilterGroupOperators_GetSelectionText() const;
 
-	FReply AddFilter_OnClicked();
-	FReply DeleteFilter_OnClicked();
+	UE_API FReply AddFilter_OnClicked();
+	UE_API FReply DeleteFilter_OnClicked();
 
-	FReply AddGroup_OnClicked();
-	FReply DeleteGroup_OnClicked();
+	UE_API FReply AddGroup_OnClicked();
+	UE_API FReply DeleteGroup_OnClicked();
 
-	FText GetTextBoxValue() const;
-	void OnTextBoxValueCommitted(const FText& InNewText, ETextCommit::Type InTextCommit);
+	UE_API FText GetTextBoxValue() const;
+	UE_API void OnTextBoxValueCommitted(const FText& InNewText, ETextCommit::Type InTextCommit);
 
-	FText GetTextBoxTooltipText() const;
-	FText GetTextBoxHintText() const;
-	bool TextBox_OnVerifyTextChanged(const FText& InText, FText& OutErrorMessage);
+	UE_API FText GetTextBoxTooltipText() const;
+	UE_API FText GetTextBoxHintText() const;
+	UE_API bool TextBox_OnVerifyTextChanged(const FText& InText, FText& OutErrorMessage);
 
-	void SuggestionTextBox_GetSuggestions(const FString& Text, TArray<FString>& Suggestions);
-	void SuggestionTextBox_GetHistory(TArray<FString>& Suggestions);
+	UE_API void SuggestionTextBox_GetSuggestions(const FString& Text, TArray<FString>& Suggestions);
+	UE_API void SuggestionTextBox_GetHistory(TArray<FString>& Suggestions);
 
-	void SuggestionTextBox_OnValueChanged(const FText& InNewText);
-	FText SuggestionTextBox_GetValue() const;
-	void SuggestionTextBox_OnValueCommitted(const FText& InNewText, ETextCommit::Type InTextCommit);
+	UE_API void SuggestionTextBox_OnValueChanged(const FText& InNewText);
+	UE_API FText SuggestionTextBox_GetValue() const;
+	UE_API void SuggestionTextBox_OnValueCommitted(const FText& InNewText, ETextCommit::Type InTextCommit);
 
 private:
 	FFilterConfiguratorNodePtr FilterConfiguratorNodePtr;
@@ -90,3 +92,5 @@ private:
 };
 
 } // namespace UE::Insights
+
+#undef UE_API
diff --git a/Engine/Source/Developer/TraceInsightsCore/Public/InsightsCore/ITraceInsightsCoreModule.h b/Engine/Source/Developer/TraceInsightsCore/Public/InsightsCore/ITraceInsightsCoreModule.h
--- a/Engine/Source/Developer/TraceInsightsCore/Public/InsightsCore/ITraceInsightsCoreModule.h
+++ b/Engine/Source/Developer/TraceInsightsCore/Public/InsightsCore/ITraceInsightsCoreModule.h
@@ -7,7 +7,7 @@
 #include "Modules/ModuleInterface.h"
 
 /** Interface for TraceInsightsCore module. */
-class TRACEINSIGHTSCORE_API ITraceInsightsCoreModule : public IModuleInterface
+class ITraceInsightsCoreModule : public IModuleInterface
 {
 public:
 	virtual ~ITraceInsightsCoreModule()
diff --git a/Engine/Source/Developer/TraceInsightsCore/Public/InsightsCore/Table/ViewModels/BaseTreeNode.h b/Engine/Source/Developer/TraceInsightsCore/Public/InsightsCore/Table/ViewModels/BaseTreeNode.h
--- a/Engine/Source/Developer/TraceInsightsCore/Public/InsightsCore/Table/ViewModels/BaseTreeNode.h
+++ b/Engine/Source/Developer/TraceInsightsCore/Public/InsightsCore/Table/ViewModels/BaseTreeNode.h
@@ -12,6 +12,8 @@
 
 #include "InsightsCore/Common/SimpleRtti.h"
 
+#define UE_API TRACEINSIGHTSCORE_API
+
 struct FSlateBrush;
 
 namespace UE::Insights
@@ -38,9 +40,9 @@ typedef TWeakPtr<class FBaseTreeNode> FBaseTreeNodeWeak;
 
 ////////////////////////////////////////////////////////////////////////////////////////////////////
 
-class TRACEINSIGHTSCORE_API FBaseTreeNode : public TSharedFromThis<FBaseTreeNode>
+class FBaseTreeNode : public TSharedFromThis<FBaseTreeNode>
 {
-	INSIGHTS_DECLARE_RTTI_BASE(FBaseTreeNode)
+	INSIGHTS_DECLARE_RTTI_BASE(FBaseTreeNode, UE_API)
 
 protected:
 	enum class EGroupNodeFlags : uint32
@@ -90,17 +92,17 @@ public:
 	/**
 	 * @returns a name of this node to display in tree view.
 	 */
-	virtual const FText GetDisplayName() const;
+	UE_API virtual const FText GetDisplayName() const;
 
 	/**
 	 * @returns a name suffix for this node to display in tree view. Ex.: group nodes may include additional info, like the number of visible / total children.
 	 */
-	virtual const FText GetExtraDisplayName() const;
+	UE_API virtual const FText GetExtraDisplayName() const;
 
 	/**
 	 * @returns true if the node has an extra display name (a name suffix to display in tree view).
 	 */
-	virtual bool HasExtraDisplayName() const;
+	UE_API virtual bool HasExtraDisplayName() const;
 
 	/**
 	 * @returns a descriptive text for this node to display in a tooltip.
@@ -113,7 +115,7 @@ public:
 	/**
 	 * @returns the default icon for a group/leaf node.
 	 */
-	static const FSlateBrush* GetDefaultIcon(bool bIsGroupNode);
+	static UE_API const FSlateBrush* GetDefaultIcon(bool bIsGroupNode);
 
 	/**
 	 * @returns the brush icon for this node.
@@ -126,7 +128,7 @@ public:
 	/**
 	 * @returns the default color tint for a group or leaf node.
 	 */
-	static FLinearColor GetDefaultIconColor(bool bIsGroupNode);
+	static UE_API FLinearColor GetDefaultIconColor(bool bIsGroupNode);
 
 	/**
 	 * @returns the color tint to be used for the icon of this node.
@@ -139,7 +141,7 @@ public:
 	/**
 	 * @returns the default color tint for a group or leaf node.
 	 */
-	static FLinearColor GetDefaultColor(bool bIsGroupNode);
+	static UE_API FLinearColor GetDefaultColor(bool bIsGroupNode);
 
 	/**
 	 * @returns the color tint to be used for the name text of this node.
@@ -237,7 +239,7 @@ public:
 		return GroupData->Children;
 	}
 
-	void SortChildren(const ITableCellValueSorter& Sorter, ESortMode SortMode);
+	UE_API void SortChildren(const ITableCellValueSorter& Sorter, ESortMode SortMode);
 
 	template <typename PredicateType>
 	void SortChildren(PredicateType Predicate)
@@ -416,7 +418,7 @@ public:
 		return *GroupData->FilteredChildrenPtr;
 	}
 
-	void SortFilteredChildren(const ITableCellValueSorter& Sorter, ESortMode SortMode);
+	UE_API void SortFilteredChildren(const ITableCellValueSorter& Sorter, ESortMode SortMode);
 
 	template <typename PredicateType>
 	void SortFilteredChildren(PredicateType Predicate)
@@ -525,9 +527,11 @@ private:
 	FGroupNodeData* GroupData;
 
 	/** The only group data for "not a group" nodes. */
-	static FGroupNodeData DefaultGroupData;
+	static UE_API FGroupNodeData DefaultGroupData;
 };
 
 ////////////////////////////////////////////////////////////////////////////////////////////////////
 
 } // namespace UE::Insights
+
+#undef UE_API
diff --git a/Engine/Source/Developer/TraceInsightsCore/Public/InsightsCore/Table/ViewModels/Table.h b/Engine/Source/Developer/TraceInsightsCore/Public/InsightsCore/Table/ViewModels/Table.h
--- a/Engine/Source/Developer/TraceInsightsCore/Public/InsightsCore/Table/ViewModels/Table.h
+++ b/Engine/Source/Developer/TraceInsightsCore/Public/InsightsCore/Table/ViewModels/Table.h
@@ -11,6 +11,8 @@
 #include "Templates/SharedPointer.h"
 #include "UObject/NameTypes.h"
 
+#define UE_API TRACEINSIGHTSCORE_API
+
 namespace UE::Insights
 {
 
@@ -22,11 +24,11 @@ class FTableColumn;
  * Table View Model.
  * View model class for the STableListView and STableTreeView widgets.
  */
-class TRACEINSIGHTSCORE_API FTable : public TSharedFromThis<FTable>
+class FTable : public TSharedFromThis<FTable>
 {
 public:
-	FTable();
-	virtual ~FTable();
+	UE_API FTable();
+	UE_API virtual ~FTable();
 
 	const FText& GetDisplayName() const { return DisplayName; }
 	void SetDisplayName(const FText& InDisplayName) { DisplayName = InDisplayName; }
@@ -34,15 +36,15 @@ public:
 	const FText& GetDescription() const { return Description; }
 	void SetDescription(const FText& InDescription) { Description = InDescription; }
 
-	virtual void Reset();
+	UE_API virtual void Reset();
 
 	int32 GetColumnCount() const { return Columns.Num(); }
 	bool IsValid() const { return Columns.Num() > 0; }
 
 	const TArray<TSharedRef<FTableColumn>>& GetColumns() const { return Columns; }
-	void SetColumns(const TArray<TSharedRef<FTableColumn>>& InColumns);
+	UE_API void SetColumns(const TArray<TSharedRef<FTableColumn>>& InColumns);
 
-	void GetVisibleColumns(TArray<TSharedRef<FTableColumn>>& InArray) const;
+	UE_API void GetVisibleColumns(TArray<TSharedRef<FTableColumn>>& InArray) const;
 
 	TSharedRef<FTableColumn> FindColumnChecked(const FName& ColumnId) const
 	{
@@ -59,16 +61,16 @@ public:
 		return nullptr;
 	}
 
-	int32 GetColumnPositionIndex(const FName& ColumnId) const;
+	UE_API int32 GetColumnPositionIndex(const FName& ColumnId) const;
 
-	void GetVisibleColumnsData(const TArray<TSharedPtr<FBaseTreeNode>>& InNodes, const FName& LogListingName, TCHAR Separator, bool bIncludeHeaders, FString& OutData) const;
+	UE_API void GetVisibleColumnsData(const TArray<TSharedPtr<FBaseTreeNode>>& InNodes, const FName& LogListingName, TCHAR Separator, bool bIncludeHeaders, FString& OutData) const;
 
-	static const FName GetHierarchyColumnId();
+	static UE_API const FName GetHierarchyColumnId();
 
 protected:
 	void ResetColumns() { Columns.Reset(); }
-	void AddColumn(TSharedRef<FTableColumn> Column);
-	void AddHierarchyColumn(int32 ColumnIndex, const TCHAR* ColumnName);
+	UE_API void AddColumn(TSharedRef<FTableColumn> Column);
+	UE_API void AddHierarchyColumn(int32 ColumnIndex, const TCHAR* ColumnName);
 
 private:
 	FText DisplayName;
@@ -84,3 +86,5 @@ private:
 ////////////////////////////////////////////////////////////////////////////////////////////////////
 
 } // namespace UE::Insights
+
+#undef UE_API
diff --git a/Engine/Source/Developer/TraceInsightsCore/Public/InsightsCore/Table/ViewModels/TableCellValue.h b/Engine/Source/Developer/TraceInsightsCore/Public/InsightsCore/Table/ViewModels/TableCellValue.h
--- a/Engine/Source/Developer/TraceInsightsCore/Public/InsightsCore/Table/ViewModels/TableCellValue.h
+++ b/Engine/Source/Developer/TraceInsightsCore/Public/InsightsCore/Table/ViewModels/TableCellValue.h
@@ -36,7 +36,7 @@ enum class TRACEINSIGHTSCORE_API ETableCellDataType : uint32
 
 ////////////////////////////////////////////////////////////////////////////////////////////////////
 
-class TRACEINSIGHTSCORE_API ICustomTableCellValue
+class ICustomTableCellValue
 {
 public:
 	virtual bool AsBool() const = 0;
@@ -47,7 +47,7 @@ public:
 
 ////////////////////////////////////////////////////////////////////////////////////////////////////
 
-class TRACEINSIGHTSCORE_API FTextCustomTableCellValue : public ICustomTableCellValue
+class FTextCustomTableCellValue : public ICustomTableCellValue
 {
 public:
 	FTextCustomTableCellValue(const FText& InText) : Text(InText) {}
@@ -66,7 +66,7 @@ private:
 
 ////////////////////////////////////////////////////////////////////////////////////////////////////
 
-struct TRACEINSIGHTSCORE_API FTableCellValue
+struct FTableCellValue
 {
 public:
 	FTableCellValue() : DataType(ETableCellDataType::Unknown) {}
diff --git a/Engine/Source/Developer/TraceInsightsCore/Public/InsightsCore/Table/ViewModels/TableCellValueFormatter.h b/Engine/Source/Developer/TraceInsightsCore/Public/InsightsCore/Table/ViewModels/TableCellValueFormatter.h
--- a/Engine/Source/Developer/TraceInsightsCore/Public/InsightsCore/Table/ViewModels/TableCellValueFormatter.h
+++ b/Engine/Source/Developer/TraceInsightsCore/Public/InsightsCore/Table/ViewModels/TableCellValueFormatter.h
@@ -9,6 +9,8 @@
 
 #include "InsightsCore/Table/ViewModels/TableCellValue.h"
 
+#define UE_API TRACEINSIGHTSCORE_API
+
 class IToolTip;
 class SWidget;
 
@@ -20,7 +22,7 @@ class FTableColumn;
 
 ////////////////////////////////////////////////////////////////////////////////////////////////////
 
-class TRACEINSIGHTSCORE_API ITableCellValueFormatter
+class ITableCellValueFormatter
 {
 public:
 	virtual FText FormatValue(const TOptional<FTableCellValue>& InValue) const = 0;
@@ -39,7 +41,7 @@ public:
 
 ////////////////////////////////////////////////////////////////////////////////////////////////////
 
-class TRACEINSIGHTSCORE_API FTableCellValueFormatter : public ITableCellValueFormatter
+class FTableCellValueFormatter : public ITableCellValueFormatter
 {
 public:
 	FTableCellValueFormatter() {}
@@ -48,22 +50,22 @@ public:
 	virtual FText FormatValue(const TOptional<FTableCellValue>& InValue) const override { return FText::GetEmpty(); }
 	virtual FText FormatValueForTooltip(const TOptional<FTableCellValue>& InValue) const override { return FormatValue(InValue); }
 
-	virtual FText FormatValue(const FTableColumn& Column, const FBaseTreeNode& Node) const override; // { return FormatValue(Column.GetValue(Node)); }
-	virtual FText FormatValueForTooltip(const FTableColumn& Column, const FBaseTreeNode& Node) const override; // { return FormatValueForTooltip(Column.GetValue(Node)); }
-	virtual FText FormatValueForGrouping(const FTableColumn& Column, const FBaseTreeNode& Node) const override; // { return FormatValueForTooltip(Column.GetValue(Node)); }
+	UE_API virtual FText FormatValue(const FTableColumn& Column, const FBaseTreeNode& Node) const override; // { return FormatValue(Column.GetValue(Node)); }
+	UE_API virtual FText FormatValueForTooltip(const FTableColumn& Column, const FBaseTreeNode& Node) const override; // { return FormatValueForTooltip(Column.GetValue(Node)); }
+	UE_API virtual FText FormatValueForGrouping(const FTableColumn& Column, const FBaseTreeNode& Node) const override; // { return FormatValueForTooltip(Column.GetValue(Node)); }
 
-	virtual FText CopyValue(const FTableColumn& Column, const FBaseTreeNode& Node) const override; // { return FormatValue(Column, Node); }
-	virtual FText CopyTooltip(const FTableColumn& Column, const FBaseTreeNode& Node) const override; // { return FormatValueForTooltip(Column, Node); }
+	UE_API virtual FText CopyValue(const FTableColumn& Column, const FBaseTreeNode& Node) const override; // { return FormatValue(Column, Node); }
+	UE_API virtual FText CopyTooltip(const FTableColumn& Column, const FBaseTreeNode& Node) const override; // { return FormatValueForTooltip(Column, Node); }
 
 	virtual TSharedPtr<SWidget> GenerateCustomWidget(const FTableColumn& Column, const FBaseTreeNode& Node) const override { return nullptr; }
-	virtual TSharedPtr<IToolTip> GetCustomTooltip(const FTableColumn& Column, const FBaseTreeNode& Node) const override;
+	UE_API virtual TSharedPtr<IToolTip> GetCustomTooltip(const FTableColumn& Column, const FBaseTreeNode& Node) const override;
 
-	static EVisibility GetTooltipVisibility();
+	static UE_API EVisibility GetTooltipVisibility();
 };
 
 ////////////////////////////////////////////////////////////////////////////////////////////////////
 
-class TRACEINSIGHTSCORE_API FTextValueFormatter : public FTableCellValueFormatter
+class FTextValueFormatter : public FTableCellValueFormatter
 {
 public:
 	virtual FText FormatValue(const TOptional<FTableCellValue>& InValue) const override
@@ -78,7 +80,7 @@ public:
 
 ////////////////////////////////////////////////////////////////////////////////////////////////////
 
-class TRACEINSIGHTSCORE_API FAsTextValueFormatter : public FTableCellValueFormatter
+class FAsTextValueFormatter : public FTableCellValueFormatter
 {
 public:
 	virtual FText FormatValue(const TOptional<FTableCellValue>& InValue) const override
@@ -93,23 +95,23 @@ public:
 
 ////////////////////////////////////////////////////////////////////////////////////////////////////
 
-class TRACEINSIGHTSCORE_API FBoolValueFormatterAsTrueFalse : public FTableCellValueFormatter
+class FBoolValueFormatterAsTrueFalse : public FTableCellValueFormatter
 {
 public:
-	virtual FText FormatValue(const TOptional<FTableCellValue>& InValue) const override;
+	UE_API virtual FText FormatValue(const TOptional<FTableCellValue>& InValue) const override;
 };
 
 ////////////////////////////////////////////////////////////////////////////////////////////////////
 
-class TRACEINSIGHTSCORE_API FBoolValueFormatterAsOnOff : public FTableCellValueFormatter
+class FBoolValueFormatterAsOnOff : public FTableCellValueFormatter
 {
 public:
-	virtual FText FormatValue(const TOptional<FTableCellValue>& InValue) const override;
+	UE_API virtual FText FormatValue(const TOptional<FTableCellValue>& InValue) const override;
 };
 
 ////////////////////////////////////////////////////////////////////////////////////////////////////
 
-class TRACEINSIGHTSCORE_API FInt64ValueFormatterAsNumber : public FTableCellValueFormatter
+class FInt64ValueFormatterAsNumber : public FTableCellValueFormatter
 {
 public:
 	virtual FText FormatValue(const TOptional<FTableCellValue>& InValue) const override
@@ -120,16 +122,16 @@ public:
 
 ////////////////////////////////////////////////////////////////////////////////////////////////////
 
-class TRACEINSIGHTSCORE_API FInt64ValueFormatterAsUInt32InfinteNumber : public FTableCellValueFormatter
+class FInt64ValueFormatterAsUInt32InfinteNumber : public FTableCellValueFormatter
 {
 public:
-	virtual FText FormatValue(const TOptional<FTableCellValue>& InValue) const override;
-	virtual FText FormatValueForTooltip(const TOptional<FTableCellValue>& InValue) const override;
+	UE_API virtual FText FormatValue(const TOptional<FTableCellValue>& InValue) const override;
+	UE_API virtual FText FormatValueForTooltip(const TOptional<FTableCellValue>& InValue) const override;
 };
 
 ////////////////////////////////////////////////////////////////////////////////////////////////////
 
-class TRACEINSIGHTSCORE_API FInt64ValueFormatterAsHex32 : public FTableCellValueFormatter
+class FInt64ValueFormatterAsHex32 : public FTableCellValueFormatter
 {
 public:
 	virtual FText FormatValue(const TOptional<FTableCellValue>& InValue) const override
@@ -140,7 +142,7 @@ public:
 
 ////////////////////////////////////////////////////////////////////////////////////////////////////
 
-class TRACEINSIGHTSCORE_API FInt64ValueFormatterAsHex64 : public FTableCellValueFormatter
+class FInt64ValueFormatterAsHex64 : public FTableCellValueFormatter
 {
 public:
 	virtual FText FormatValue(const TOptional<FTableCellValue>& InValue) const override
@@ -151,7 +153,7 @@ public:
 
 ////////////////////////////////////////////////////////////////////////////////////////////////////
 
-class TRACEINSIGHTSCORE_API FInt64ValueFormatterAsMemory : public FTableCellValueFormatter
+class FInt64ValueFormatterAsMemory : public FTableCellValueFormatter
 {
 public:
 	FInt64ValueFormatterAsMemory()
@@ -163,10 +165,10 @@ public:
 	const FNumberFormattingOptions& GetFormattingOptions() const { return FormattingOptions; }
 	FNumberFormattingOptions& GetFormattingOptions() { return FormattingOptions; }
 
-	virtual FText FormatValue(const TOptional<FTableCellValue>& InValue) const override;
-	virtual FText FormatValueForTooltip(const TOptional<FTableCellValue>& InValue) const override;
+	UE_API virtual FText FormatValue(const TOptional<FTableCellValue>& InValue) const override;
+	UE_API virtual FText FormatValueForTooltip(const TOptional<FTableCellValue>& InValue) const override;
 
-	static FText FormatForTooltip(int64 InValue);
+	static UE_API FText FormatForTooltip(int64 InValue);
 
 private:
 	FNumberFormattingOptions FormattingOptions;
@@ -174,55 +176,57 @@ private:
 
 ////////////////////////////////////////////////////////////////////////////////////////////////////
 
-class TRACEINSIGHTSCORE_API FFloatValueFormatterAsNumber : public FTableCellValueFormatter
+class FFloatValueFormatterAsNumber : public FTableCellValueFormatter
 {
 public:
-	virtual FText FormatValue(const TOptional<FTableCellValue>& InValue) const override;
+	UE_API virtual FText FormatValue(const TOptional<FTableCellValue>& InValue) const override;
 };
 
 ////////////////////////////////////////////////////////////////////////////////////////////////////
 
-class TRACEINSIGHTSCORE_API FFloatValueFormatterAsTimeAuto : public FTableCellValueFormatter
+class FFloatValueFormatterAsTimeAuto : public FTableCellValueFormatter
 {
 public:
-	virtual FText FormatValue(const TOptional<FTableCellValue>& InValue) const override;
-	virtual FText FormatValueForTooltip(const TOptional<FTableCellValue>& InValue) const override;
+	UE_API virtual FText FormatValue(const TOptional<FTableCellValue>& InValue) const override;
+	UE_API virtual FText FormatValueForTooltip(const TOptional<FTableCellValue>& InValue) const override;
 };
 
 ////////////////////////////////////////////////////////////////////////////////////////////////////
 
-class TRACEINSIGHTSCORE_API FDoubleValueFormatterAsNumber : public FTableCellValueFormatter
+class FDoubleValueFormatterAsNumber : public FTableCellValueFormatter
 {
 public:
-	virtual FText FormatValue(const TOptional<FTableCellValue>& InValue) const override;
+	UE_API virtual FText FormatValue(const TOptional<FTableCellValue>& InValue) const override;
 };
 
 ////////////////////////////////////////////////////////////////////////////////////////////////////
 
-class TRACEINSIGHTSCORE_API FDoubleValueFormatterAsTimeAuto : public FTableCellValueFormatter
+class FDoubleValueFormatterAsTimeAuto : public FTableCellValueFormatter
 {
 public:
-	virtual FText FormatValue(const TOptional<FTableCellValue>& InValue) const override;
-	virtual FText FormatValueForTooltip(const TOptional<FTableCellValue>& InValue) const override;
+	UE_API virtual FText FormatValue(const TOptional<FTableCellValue>& InValue) const override;
+	UE_API virtual FText FormatValueForTooltip(const TOptional<FTableCellValue>& InValue) const override;
 };
 
 ////////////////////////////////////////////////////////////////////////////////////////////////////
 
-class TRACEINSIGHTSCORE_API FDoubleValueFormatterAsTimeMs : public FTableCellValueFormatter
+class FDoubleValueFormatterAsTimeMs : public FTableCellValueFormatter
 {
 public:
-	virtual FText FormatValue(const TOptional<FTableCellValue>& InValue) const override;
-	virtual FText FormatValueForTooltip(const TOptional<FTableCellValue>& InValue) const override;
+	UE_API virtual FText FormatValue(const TOptional<FTableCellValue>& InValue) const override;
+	UE_API virtual FText FormatValueForTooltip(const TOptional<FTableCellValue>& InValue) const override;
 };
 
 ////////////////////////////////////////////////////////////////////////////////////////////////////
 
-class TRACEINSIGHTSCORE_API FCStringValueFormatterAsText : public FTableCellValueFormatter
+class FCStringValueFormatterAsText : public FTableCellValueFormatter
 {
 public:
-	virtual FText FormatValue(const TOptional<FTableCellValue>& InValue) const override;
+	UE_API virtual FText FormatValue(const TOptional<FTableCellValue>& InValue) const override;
 };
 
 ////////////////////////////////////////////////////////////////////////////////////////////////////
 
 } // namespace UE::Insights
+
+#undef UE_API
diff --git a/Engine/Source/Developer/TraceInsightsCore/Public/InsightsCore/Table/ViewModels/TableCellValueGetter.h b/Engine/Source/Developer/TraceInsightsCore/Public/InsightsCore/Table/ViewModels/TableCellValueGetter.h
--- a/Engine/Source/Developer/TraceInsightsCore/Public/InsightsCore/Table/ViewModels/TableCellValueGetter.h
+++ b/Engine/Source/Developer/TraceInsightsCore/Public/InsightsCore/Table/ViewModels/TableCellValueGetter.h
@@ -8,6 +8,8 @@
 
 #include "InsightsCore/Table/ViewModels/TableCellValue.h"
 
+#define UE_API TRACEINSIGHTSCORE_API
+
 namespace UE::Insights
 {
 
@@ -16,7 +18,7 @@ class FTableColumn;
 
 ////////////////////////////////////////////////////////////////////////////////////////////////////
 
-class TRACEINSIGHTSCORE_API ITableCellValueGetter
+class ITableCellValueGetter
 {
 public:
 	virtual const TOptional<FTableCellValue> GetValue(const FTableColumn& Column, const FBaseTreeNode& Node) const = 0;
@@ -25,7 +27,7 @@ public:
 
 ////////////////////////////////////////////////////////////////////////////////////////////////////
 
-class TRACEINSIGHTSCORE_API FTableCellValueGetter : public ITableCellValueGetter
+class FTableCellValueGetter : public ITableCellValueGetter
 {
 public:
 	FTableCellValueGetter() {}
@@ -37,20 +39,22 @@ public:
 
 ////////////////////////////////////////////////////////////////////////////////////////////////////
 
-class TRACEINSIGHTSCORE_API FNameValueGetter : public FTableCellValueGetter
+class FNameValueGetter : public FTableCellValueGetter
 {
 public:
-	virtual const TOptional<FTableCellValue> GetValue(const FTableColumn& Column, const FBaseTreeNode& Node) const override;
+	UE_API virtual const TOptional<FTableCellValue> GetValue(const FTableColumn& Column, const FBaseTreeNode& Node) const override;
 };
 
 ////////////////////////////////////////////////////////////////////////////////////////////////////
 
-class TRACEINSIGHTSCORE_API FDisplayNameValueGetter : public FTableCellValueGetter
+class FDisplayNameValueGetter : public FTableCellValueGetter
 {
 public:
-	virtual const TOptional<FTableCellValue> GetValue(const FTableColumn& Column, const FBaseTreeNode& Node) const override;
+	UE_API virtual const TOptional<FTableCellValue> GetValue(const FTableColumn& Column, const FBaseTreeNode& Node) const override;
 };
 
 ////////////////////////////////////////////////////////////////////////////////////////////////////
 
 } // namespace UE::Insights
+
+#undef UE_API
diff --git a/Engine/Source/Developer/TraceInsightsCore/Public/InsightsCore/Table/ViewModels/TableCellValueSorter.h b/Engine/Source/Developer/TraceInsightsCore/Public/InsightsCore/Table/ViewModels/TableCellValueSorter.h
--- a/Engine/Source/Developer/TraceInsightsCore/Public/InsightsCore/Table/ViewModels/TableCellValueSorter.h
+++ b/Engine/Source/Developer/TraceInsightsCore/Public/InsightsCore/Table/ViewModels/TableCellValueSorter.h
@@ -8,6 +8,8 @@
 #include "InsightsCore/Table/ViewModels/TableCellValue.h"
 #include "InsightsCore/Table/ViewModels/TableColumn.h"
 
+#define UE_API TRACEINSIGHTSCORE_API
+
 struct FSlateBrush;
 
 namespace UE::Insights
@@ -28,7 +30,7 @@ enum class TRACEINSIGHTSCORE_API ESortMode
 
 typedef TFunction<bool(const FBaseTreeNodePtr& A, const FBaseTreeNodePtr& B)> FTreeNodeCompareFunc;
 
-class TRACEINSIGHTSCORE_API ITableCellValueSorter
+class ITableCellValueSorter
 {
 public:
 	virtual FName GetName() const = 0;
@@ -48,10 +50,10 @@ public:
 
 ////////////////////////////////////////////////////////////////////////////////////////////////////
 
-class TRACEINSIGHTSCORE_API FTableCellValueSorter : public ITableCellValueSorter
+class FTableCellValueSorter : public ITableCellValueSorter
 {
 public:
-	FTableCellValueSorter(const FName InName, const FText& InShortName, const FText& InTitleName, const FText& InDescription, TSharedRef<FTableColumn> InColumnRef);
+	UE_API FTableCellValueSorter(const FName InName, const FText& InShortName, const FText& InTitleName, const FText& InDescription, TSharedRef<FTableColumn> InColumnRef);
 	virtual ~FTableCellValueSorter() {}
 
 	virtual FName GetName() const override { return Name; }
@@ -63,14 +65,14 @@ public:
 	virtual FSlateBrush* GetIcon(ESortMode SortMode) const override { return SortMode == ESortMode::Ascending ? AscendingIcon : DescendingIcon; }
 	virtual FTreeNodeCompareFunc GetTreeNodeCompareDelegate(ESortMode SortMode) const override { return SortMode == ESortMode::Ascending ? AscendingCompareDelegate : DescendingCompareDelegate; }
 
-	virtual void Sort(TArray<FBaseTreeNodePtr>& NodesToSort, ESortMode SortMode) const override;
+	UE_API virtual void Sort(TArray<FBaseTreeNodePtr>& NodesToSort, ESortMode SortMode) const override;
 
 	virtual void SetAsyncOperationProgress(IAsyncOperationProgress* InAsyncOperationProgress) override { AsyncOperationProgress = InAsyncOperationProgress; }
 
 protected:
 	// Attempts to cancel the sort by throwing an exception. If exceptions are not available the return value is meant to be returned from sort predicates to speed up the sort.
-	bool CancelSort() const;
-	bool ShouldCancelSort() const;
+	UE_API bool CancelSort() const;
+	UE_API bool ShouldCancelSort() const;
 
 protected:
 	FName Name;
@@ -91,94 +93,96 @@ protected:
 
 ////////////////////////////////////////////////////////////////////////////////////////////////////
 
-class TRACEINSIGHTSCORE_API FBaseTableColumnSorter : public FTableCellValueSorter
+class FBaseTableColumnSorter : public FTableCellValueSorter
 {
 public:
-	FBaseTableColumnSorter(TSharedRef<FTableColumn> InColumnRef);
+	UE_API FBaseTableColumnSorter(TSharedRef<FTableColumn> InColumnRef);
 };
 
 ////////////////////////////////////////////////////////////////////////////////////////////////////
 
-class TRACEINSIGHTSCORE_API FSorterByName : public FTableCellValueSorter
+class FSorterByName : public FTableCellValueSorter
 {
 public:
-	FSorterByName(TSharedRef<FTableColumn> InColumnRef);
+	UE_API FSorterByName(TSharedRef<FTableColumn> InColumnRef);
 };
 
 ////////////////////////////////////////////////////////////////////////////////////////////////////
 
-class TRACEINSIGHTSCORE_API FSorterByTypeName : public FTableCellValueSorter
+class FSorterByTypeName : public FTableCellValueSorter
 {
 public:
-	FSorterByTypeName(TSharedRef<FTableColumn> InColumnRef);
+	UE_API FSorterByTypeName(TSharedRef<FTableColumn> InColumnRef);
 };
 
 ////////////////////////////////////////////////////////////////////////////////////////////////////
 
-class TRACEINSIGHTSCORE_API FSorterByBoolValue : public FBaseTableColumnSorter
+class FSorterByBoolValue : public FBaseTableColumnSorter
 {
 public:
-	FSorterByBoolValue(TSharedRef<FTableColumn> InColumnRef);
+	UE_API FSorterByBoolValue(TSharedRef<FTableColumn> InColumnRef);
 };
 
 ////////////////////////////////////////////////////////////////////////////////////////////////////
 
-class TRACEINSIGHTSCORE_API FSorterByInt64Value : public FBaseTableColumnSorter
+class FSorterByInt64Value : public FBaseTableColumnSorter
 {
 public:
-	FSorterByInt64Value(TSharedRef<FTableColumn> InColumnRef);
+	UE_API FSorterByInt64Value(TSharedRef<FTableColumn> InColumnRef);
 
-	virtual void Sort(TArray<FBaseTreeNodePtr>& NodesToSort, ESortMode SortMode) const override;
+	UE_API virtual void Sort(TArray<FBaseTreeNodePtr>& NodesToSort, ESortMode SortMode) const override;
 };
 
 ////////////////////////////////////////////////////////////////////////////////////////////////////
 
-class TRACEINSIGHTSCORE_API FSorterByFloatValue : public FBaseTableColumnSorter
+class FSorterByFloatValue : public FBaseTableColumnSorter
 {
 public:
-	FSorterByFloatValue(TSharedRef<FTableColumn> InColumnRef);
+	UE_API FSorterByFloatValue(TSharedRef<FTableColumn> InColumnRef);
 
-	virtual void Sort(TArray<FBaseTreeNodePtr>& NodesToSort, ESortMode SortMode) const override;
+	UE_API virtual void Sort(TArray<FBaseTreeNodePtr>& NodesToSort, ESortMode SortMode) const override;
 };
 
 ////////////////////////////////////////////////////////////////////////////////////////////////////
 
-class TRACEINSIGHTSCORE_API FSorterByDoubleValue : public FBaseTableColumnSorter
+class FSorterByDoubleValue : public FBaseTableColumnSorter
 {
 public:
-	FSorterByDoubleValue(TSharedRef<FTableColumn> InColumnRef);
+	UE_API FSorterByDoubleValue(TSharedRef<FTableColumn> InColumnRef);
 
-	virtual void Sort(TArray<FBaseTreeNodePtr>& NodesToSort, ESortMode SortMode) const override;
+	UE_API virtual void Sort(TArray<FBaseTreeNodePtr>& NodesToSort, ESortMode SortMode) const override;
 };
 
 ////////////////////////////////////////////////////////////////////////////////////////////////////
 
-class TRACEINSIGHTSCORE_API FSorterByCStringValue : public FBaseTableColumnSorter
+class FSorterByCStringValue : public FBaseTableColumnSorter
 {
 public:
-	FSorterByCStringValue(TSharedRef<FTableColumn> InColumnRef);
+	UE_API FSorterByCStringValue(TSharedRef<FTableColumn> InColumnRef);
 };
 
 ////////////////////////////////////////////////////////////////////////////////////////////////////
 
-class TRACEINSIGHTSCORE_API FSorterByTextValue : public FBaseTableColumnSorter
+class FSorterByTextValue : public FBaseTableColumnSorter
 {
 public:
-	FSorterByTextValue(TSharedRef<FTableColumn> InColumnRef);
+	UE_API FSorterByTextValue(TSharedRef<FTableColumn> InColumnRef);
 
-	virtual void Sort(TArray<FBaseTreeNodePtr>& NodesToSort, ESortMode SortMode) const override;
+	UE_API virtual void Sort(TArray<FBaseTreeNodePtr>& NodesToSort, ESortMode SortMode) const override;
 };
 
 ////////////////////////////////////////////////////////////////////////////////////////////////////
 
-class TRACEINSIGHTSCORE_API FSorterByTextValueWithId : public FBaseTableColumnSorter
+class FSorterByTextValueWithId : public FBaseTableColumnSorter
 {
 public:
-	FSorterByTextValueWithId(TSharedRef<FTableColumn> InColumnRef);
+	UE_API FSorterByTextValueWithId(TSharedRef<FTableColumn> InColumnRef);
 
-	virtual void Sort(TArray<FBaseTreeNodePtr>& NodesToSort, ESortMode SortMode) const override;
+	UE_API virtual void Sort(TArray<FBaseTreeNodePtr>& NodesToSort, ESortMode SortMode) const override;
 };
 
 ////////////////////////////////////////////////////////////////////////////////////////////////////
 
 } // namespace UE::Insights
+
+#undef UE_API
diff --git a/Engine/Source/Developer/TraceInsightsCore/Public/InsightsCore/Table/ViewModels/TableColumn.h b/Engine/Source/Developer/TraceInsightsCore/Public/InsightsCore/Table/ViewModels/TableColumn.h
--- a/Engine/Source/Developer/TraceInsightsCore/Public/InsightsCore/Table/ViewModels/TableColumn.h
+++ b/Engine/Source/Developer/TraceInsightsCore/Public/InsightsCore/Table/ViewModels/TableColumn.h
@@ -10,6 +10,8 @@
 
 #include "InsightsCore/Table/ViewModels/TableCellValue.h"
 
+#define UE_API TRACEINSIGHTSCORE_API
+
 namespace UE::Insights
 {
 
@@ -49,7 +51,7 @@ enum class TRACEINSIGHTSCORE_API ETableColumnAggregation : uint32
  * Table Column View Model.
  * Holds information about a column in STableListView or STableTreeView widgets.
  */
-class TRACEINSIGHTSCORE_API FTableColumn
+class FTableColumn
 {
 public:
 	/** No default constructor. */
@@ -160,30 +162,30 @@ public:
 	//////////////////////////////////////////////////
 	// Value Getter
 
-	static TSharedRef<ITableCellValueGetter> GetDefaultValueGetter();
+	static UE_API TSharedRef<ITableCellValueGetter> GetDefaultValueGetter();
 
 	TSharedRef<ITableCellValueGetter> GetValueGetter() const { return ValueGetter; }
 	void SetValueGetter(TSharedRef<ITableCellValueGetter> InValueGetter) { ValueGetter = InValueGetter; }
 
-	const TOptional<FTableCellValue> GetValue(const FBaseTreeNode& InNode) const;
-	uint64 GetValueId(const FBaseTreeNode& InNode) const;
+	UE_API const TOptional<FTableCellValue> GetValue(const FBaseTreeNode& InNode) const;
+	UE_API uint64 GetValueId(const FBaseTreeNode& InNode) const;
 
 	//////////////////////////////////////////////////
 	// Value Formatter
 
-	static TSharedRef<ITableCellValueFormatter> GetDefaultValueFormatter();
+	static UE_API TSharedRef<ITableCellValueFormatter> GetDefaultValueFormatter();
 
 	TSharedRef<ITableCellValueFormatter> GetValueFormatter() const { return ValueFormatter; }
 	void SetValueFormatter(TSharedRef<ITableCellValueFormatter> InValueFormatter) { ValueFormatter = InValueFormatter; }
 
-	FText GetValueAsText(const FBaseTreeNode& InNode) const;
-	FText GetValueAsTooltipText(const FBaseTreeNode& InNode) const;
-	FText GetValueAsGroupingText(const FBaseTreeNode& InNode) const;
+	UE_API FText GetValueAsText(const FBaseTreeNode& InNode) const;
+	UE_API FText GetValueAsTooltipText(const FBaseTreeNode& InNode) const;
+	UE_API FText GetValueAsGroupingText(const FBaseTreeNode& InNode) const;
 
-	FText CopyValue(const FBaseTreeNode& InNode) const;
-	FText CopyTooltip(const FBaseTreeNode& InNode) const;
+	UE_API FText CopyValue(const FBaseTreeNode& InNode) const;
+	UE_API FText CopyTooltip(const FBaseTreeNode& InNode) const;
 
-	FString GetValueAsSerializableString(const FBaseTreeNode& InNode) const;
+	UE_API FString GetValueAsSerializableString(const FBaseTreeNode& InNode) const;
 
 	//////////////////////////////////////////////////
 	// Value Sorter (can be nullptr)
@@ -271,3 +273,5 @@ private:
 ////////////////////////////////////////////////////////////////////////////////////////////////////
 
 } // namespace UE::Insights
+
+#undef UE_API
diff --git a/Engine/Source/Developer/TraceInsightsCore/Public/InsightsCore/Table/ViewModels/TableImporter.h b/Engine/Source/Developer/TraceInsightsCore/Public/InsightsCore/Table/ViewModels/TableImporter.h
--- a/Engine/Source/Developer/TraceInsightsCore/Public/InsightsCore/Table/ViewModels/TableImporter.h
+++ b/Engine/Source/Developer/TraceInsightsCore/Public/InsightsCore/Table/ViewModels/TableImporter.h
@@ -8,6 +8,8 @@
 #include "TraceServices/Model/TableImport.h"
 #include "TraceServices/Containers/Tables.h"
 
+#define UE_API TRACEINSIGHTSCORE_API
+
 class SDockTab;
 class FSpawnTabArgs;
 
@@ -23,7 +25,7 @@ class SUntypedTableTreeView;
 
 ////////////////////////////////////////////////////////////////////////////////////////////////////
 
-class TRACEINSIGHTSCORE_API FTableImporter : public TSharedFromThis<FTableImporter>
+class FTableImporter : public TSharedFromThis<FTableImporter>
 {
 private:
 	struct FOpenTableTabData
@@ -38,27 +40,27 @@ private:
 	};
 
 public:
-	FTableImporter(FName InLogListingName);
-	virtual ~FTableImporter();
+	UE_API FTableImporter(FName InLogListingName);
+	UE_API virtual ~FTableImporter();
 
-	void StartImportProcess();
-	void ImportFile(const FString& Filename);
+	UE_API void StartImportProcess();
+	UE_API void ImportFile(const FString& Filename);
 
-	void StartDiffProcess();
-	void DiffFiles(const FString& FilenameA, const FString& FilenameB);
+	UE_API void StartDiffProcess();
+	UE_API void DiffFiles(const FString& FilenameA, const FString& FilenameB);
 
-	void CloseAllOpenTabs();
+	UE_API void CloseAllOpenTabs();
 
 private:
-	TSharedRef<SDockTab> SpawnTab_TableImportTreeView(const FSpawnTabArgs& Args, FName TableViewID, FText InDisplayName);
-	TSharedRef<SDockTab> SpawnTab_TableDiffTreeView(const FSpawnTabArgs& Args, FName TableViewID, FText InDisplayName);
-	void OnTableImportTreeViewTabClosed(TSharedRef<SDockTab> TabBeingClosed);
+	UE_API TSharedRef<SDockTab> SpawnTab_TableImportTreeView(const FSpawnTabArgs& Args, FName TableViewID, FText InDisplayName);
+	UE_API TSharedRef<SDockTab> SpawnTab_TableDiffTreeView(const FSpawnTabArgs& Args, FName TableViewID, FText InDisplayName);
+	UE_API void OnTableImportTreeViewTabClosed(TSharedRef<SDockTab> TabBeingClosed);
 
-	void DisplayImportTable(FName TableViewID);
-	void DisplayDiffTable(FName TableViewID);
-	void TableImportServiceCallback(TSharedPtr<TraceServices::FTableImportCallbackParams> Params);
+	UE_API void DisplayImportTable(FName TableViewID);
+	UE_API void DisplayDiffTable(FName TableViewID);
+	UE_API void TableImportServiceCallback(TSharedPtr<TraceServices::FTableImportCallbackParams> Params);
 
-	FName GetTableID(const FString& Path);
+	UE_API FName GetTableID(const FString& Path);
 
 private:
 	FName LogListingName;
@@ -68,3 +70,5 @@ private:
 ////////////////////////////////////////////////////////////////////////////////////////////////////
 
 } // namespace UE::Insights
+
+#undef UE_API
diff --git a/Engine/Source/Developer/TraceInsightsCore/Public/InsightsCore/Table/ViewModels/TableTreeNode.h b/Engine/Source/Developer/TraceInsightsCore/Public/InsightsCore/Table/ViewModels/TableTreeNode.h
--- a/Engine/Source/Developer/TraceInsightsCore/Public/InsightsCore/Table/ViewModels/TableTreeNode.h
+++ b/Engine/Source/Developer/TraceInsightsCore/Public/InsightsCore/Table/ViewModels/TableTreeNode.h
@@ -7,6 +7,8 @@
 #include "InsightsCore/Table/ViewModels/BaseTreeNode.h"
 #include "InsightsCore/Table/ViewModels/TableCellValue.h"
 
+#define UE_API TRACEINSIGHTSCORE_API
+
 namespace UE::Insights
 {
 
@@ -14,7 +16,7 @@ class FTable;
 class FTreeNodeGrouping;
 class STableTreeView;
 
-struct TRACEINSIGHTSCORE_API FTableRowId
+struct FTableRowId
 {
 	static constexpr int32 InvalidRowIndex = -1;
 
@@ -46,9 +48,9 @@ typedef TWeakPtr<class FTableTreeNode> FTableTreeNodeWeak;
  * Table Tree Node View Model.
  * Class used to store information about a generic table tree node (used in STableTreeView).
  */
-class TRACEINSIGHTSCORE_API FTableTreeNode : public FBaseTreeNode
+class FTableTreeNode : public FBaseTreeNode
 {
-	INSIGHTS_DECLARE_RTTI(FTableTreeNode, FBaseTreeNode)
+	INSIGHTS_DECLARE_RTTI(FTableTreeNode, FBaseTreeNode, UE_API)
 
 public:
 	/** Initialization constructor for a table record node. */
@@ -167,9 +169,9 @@ private:
 
 ////////////////////////////////////////////////////////////////////////////////////////////////////
 
-class TRACEINSIGHTSCORE_API FCustomTableTreeNode : public FTableTreeNode
+class FCustomTableTreeNode : public FTableTreeNode
 {
-	INSIGHTS_DECLARE_RTTI(FCustomTableTreeNode, FTableTreeNode)
+	INSIGHTS_DECLARE_RTTI(FCustomTableTreeNode, FTableTreeNode, UE_API)
 
 public:
 	/** Initialization constructor for a table record node. */
@@ -258,3 +260,5 @@ private:
 ////////////////////////////////////////////////////////////////////////////////////////////////////
 
 } // namespace UE::Insights
+
+#undef UE_API
diff --git a/Engine/Source/Developer/TraceInsightsCore/Public/InsightsCore/Table/ViewModels/UntypedTable.h b/Engine/Source/Developer/TraceInsightsCore/Public/InsightsCore/Table/ViewModels/UntypedTable.h
--- a/Engine/Source/Developer/TraceInsightsCore/Public/InsightsCore/Table/ViewModels/UntypedTable.h
+++ b/Engine/Source/Developer/TraceInsightsCore/Public/InsightsCore/Table/ViewModels/UntypedTable.h
@@ -9,6 +9,8 @@
 // TraceInsightsCore
 #include "InsightsCore/Table/ViewModels/Table.h"
 
+#define UE_API TRACEINSIGHTSCORE_API
+
 namespace TraceServices
 {
 	class ITableLayout;
@@ -21,22 +23,22 @@ namespace UE::Insights
 
 ////////////////////////////////////////////////////////////////////////////////////////////////////
 
-class TRACEINSIGHTSCORE_API FUntypedTable : public FTable
+class FUntypedTable : public FTable
 {
 public:
-	FUntypedTable();
-	virtual ~FUntypedTable();
+	UE_API FUntypedTable();
+	UE_API virtual ~FUntypedTable();
 
-	virtual void Reset();
+	UE_API virtual void Reset();
 
 	TSharedPtr<TraceServices::IUntypedTable> GetSourceTable() const { return SourceTable; }
 	TSharedPtr<TraceServices::IUntypedTableReader> GetTableReader() const { return TableReader; }
 
 	/* Update table content. Returns true if the table layout has changed. */
-	bool UpdateSourceTable(TSharedPtr<TraceServices::IUntypedTable> InSourceTable);
+	UE_API bool UpdateSourceTable(TSharedPtr<TraceServices::IUntypedTable> InSourceTable);
 
 private:
-	void CreateColumns(const TraceServices::ITableLayout& TableLayout);
+	UE_API void CreateColumns(const TraceServices::ITableLayout& TableLayout);
 
 private:
 	TSharedPtr<TraceServices::IUntypedTable> SourceTable;
@@ -46,3 +48,5 @@ private:
 ////////////////////////////////////////////////////////////////////////////////////////////////////
 
 } // namespace UE::Insights
+
+#undef UE_API
diff --git a/Engine/Source/Developer/TraceInsightsCore/Public/InsightsCore/Table/Widgets/SAsyncOperationStatus.h b/Engine/Source/Developer/TraceInsightsCore/Public/InsightsCore/Table/Widgets/SAsyncOperationStatus.h
--- a/Engine/Source/Developer/TraceInsightsCore/Public/InsightsCore/Table/Widgets/SAsyncOperationStatus.h
+++ b/Engine/Source/Developer/TraceInsightsCore/Public/InsightsCore/Table/Widgets/SAsyncOperationStatus.h
@@ -7,34 +7,38 @@
 #include "Widgets/DeclarativeSyntaxSupport.h"
 #include "Widgets/SCompoundWidget.h"
 
+#define UE_API TRACEINSIGHTSCORE_API
+
 namespace UE::Insights
 {
 
 class IAsyncOperationStatusProvider;
 
 // A widget representing the status message used in table / tree views for background operations.
-class TRACEINSIGHTSCORE_API SAsyncOperationStatus : public SCompoundWidget
+class SAsyncOperationStatus : public SCompoundWidget
 {
 public:
 	SLATE_BEGIN_ARGS(SAsyncOperationStatus) {}
 
 	SLATE_END_ARGS()
 
-	void Construct(const FArguments& InArgs, TSharedRef<IAsyncOperationStatusProvider> InStatusProvider);
+	UE_API void Construct(const FArguments& InArgs, TSharedRef<IAsyncOperationStatusProvider> InStatusProvider);
 
 private:
-	EVisibility GetContentVisibility() const;
+	UE_API EVisibility GetContentVisibility() const;
 
-	FSlateColor GetBackgroundColorAndOpacity() const;
-	FSlateColor GetTextColorAndOpacity() const;
-	float ComputeOpacity() const;
+	UE_API FSlateColor GetBackgroundColorAndOpacity() const;
+	UE_API FSlateColor GetTextColorAndOpacity() const;
+	UE_API float ComputeOpacity() const;
 
-	FText GetText() const;
-	FText GetAnimatedText() const;
-	FText GetTooltipText() const;
+	UE_API FText GetText() const;
+	UE_API FText GetAnimatedText() const;
+	UE_API FText GetTooltipText() const;
 
 private:
 	TWeakPtr<IAsyncOperationStatusProvider> StatusProvider;
 };
 
 } // namespace UE::Insights
+
+#undef UE_API
diff --git a/Engine/Source/Developer/TraceInsightsCore/Public/InsightsCore/Table/Widgets/STableTreeViewCell.h b/Engine/Source/Developer/TraceInsightsCore/Public/InsightsCore/Table/Widgets/STableTreeViewCell.h
--- a/Engine/Source/Developer/TraceInsightsCore/Public/InsightsCore/Table/Widgets/STableTreeViewCell.h
+++ b/Engine/Source/Developer/TraceInsightsCore/Public/InsightsCore/Table/Widgets/STableTreeViewCell.h
@@ -10,6 +10,8 @@
 #include "InsightsCore/Table/ViewModels/TableTreeNode.h"
 #include "InsightsCore/Table/Widgets/STableTreeViewRow.h"
 
+#define UE_API TRACEINSIGHTSCORE_API
+
 class ITableRow;
 class IToolTip;
 
@@ -21,7 +23,7 @@ class FTableColumn;
 
 DECLARE_DELEGATE_ThreeParams(FSetHoveredTableTreeViewCell, TSharedPtr<FTable> /*TablePtr*/, TSharedPtr<FTableColumn> /*ColumnPtr*/, FTableTreeNodePtr /*TableTreeNodePtr*/);
 
-class TRACEINSIGHTSCORE_API STableTreeViewCell : public SCompoundWidget
+class STableTreeViewCell : public SCompoundWidget
 {
 public:
 	SLATE_BEGIN_ARGS(STableTreeViewCell) {}
@@ -32,12 +34,12 @@ public:
 		SLATE_ARGUMENT(FTableTreeNodePtr, TableTreeNodePtr)
 	SLATE_END_ARGS()
 
-	void Construct(const FArguments& InArgs, const TSharedRef<ITableRow>& InTableRow);
+	UE_API void Construct(const FArguments& InArgs, const TSharedRef<ITableRow>& InTableRow);
 
 protected:
-	TSharedRef<SWidget> GenerateWidgetForColumn(const FArguments& InArgs);
-	TSharedRef<SWidget> GenerateWidgetForNameColumn(const FArguments& InArgs);
-	TSharedRef<SWidget> GenerateWidgetForTableColumn(const FArguments& InArgs);
+	UE_API TSharedRef<SWidget> GenerateWidgetForColumn(const FArguments& InArgs);
+	UE_API TSharedRef<SWidget> GenerateWidgetForNameColumn(const FArguments& InArgs);
+	UE_API TSharedRef<SWidget> GenerateWidgetForTableColumn(const FArguments& InArgs);
 
 	/**
 	 * The system will use this event to notify a widget that the cursor has entered it. This event is NOT bubbled.
@@ -101,7 +103,7 @@ protected:
 		return IsHovered() ? EVisibility::Visible : EVisibility::Hidden;
 	}
 
-	bool IsSelected() const;
+	UE_API bool IsSelected() const;
 
 	const FSlateBrush* GetIcon() const
 	{
@@ -123,13 +125,13 @@ protected:
 		return TableTreeNodePtr->HasExtraDisplayName() ? EVisibility::Visible : EVisibility::Collapsed;
 	}
 
-	FText GetValueAsText() const;
+	UE_API FText GetValueAsText() const;
 
-	FSlateColor GetIconColorAndOpacity() const;
-	FSlateColor GetDisplayNameColorAndOpacity() const;
-	FSlateColor GetExtraDisplayNameColorAndOpacity() const;
-	FSlateColor GetNormalTextColorAndOpacity() const;
-	FLinearColor GetShadowColorAndOpacity() const;
+	UE_API FSlateColor GetIconColorAndOpacity() const;
+	UE_API FSlateColor GetDisplayNameColorAndOpacity() const;
+	UE_API FSlateColor GetExtraDisplayNameColorAndOpacity() const;
+	UE_API FSlateColor GetNormalTextColorAndOpacity() const;
+	UE_API FLinearColor GetShadowColorAndOpacity() const;
 
 protected:
 	TSharedPtr<ITableRow> TableRow;
@@ -147,3 +149,5 @@ protected:
 };
 
 } // namespace UE::Insights
+
+#undef UE_API
diff --git a/Engine/Source/Developer/TraceInsightsCore/Public/InsightsCore/Table/Widgets/STableTreeViewRow.h b/Engine/Source/Developer/TraceInsightsCore/Public/InsightsCore/Table/Widgets/STableTreeViewRow.h
--- a/Engine/Source/Developer/TraceInsightsCore/Public/InsightsCore/Table/Widgets/STableTreeViewRow.h
+++ b/Engine/Source/Developer/TraceInsightsCore/Public/InsightsCore/Table/Widgets/STableTreeViewRow.h
@@ -12,6 +12,8 @@
 
 #include "InsightsCore/Table/ViewModels/TableTreeNode.h"
 
+#define UE_API TRACEINSIGHTSCORE_API
+
 class IToolTip;
 
 namespace UE::Insights
@@ -27,7 +29,7 @@ DECLARE_DELEGATE_RetVal_OneParam(EHorizontalAlignment, FGetColumnOutlineHAlignme
 DECLARE_DELEGATE_ThreeParams(FSetHoveredTableTreeViewCell, TSharedPtr<FTable> /*TablePtr*/, TSharedPtr<FTableColumn> /*ColumnPtr*/, FTableTreeNodePtr /*TableTreeNodePtr*/);
 
 /** Widget that represents a table row in the tree control. Generates widgets for each column on demand. */
-class TRACEINSIGHTSCORE_API STableTreeViewRow : public SMultiColumnTableRow<FTableTreeNodePtr>
+class STableTreeViewRow : public SMultiColumnTableRow<FTableTreeNodePtr>
 {
 public:
 	SLATE_BEGIN_ARGS(STableTreeViewRow) {}
@@ -42,9 +44,9 @@ public:
 	SLATE_END_ARGS()
 
 public:
-	void Construct(const FArguments& InArgs, const TSharedRef<STableViewBase>& InOwnerTableView);
+	UE_API void Construct(const FArguments& InArgs, const TSharedRef<STableViewBase>& InOwnerTableView);
 
-	virtual TSharedRef<SWidget> GenerateWidgetForColumn(const FName& ColumnId) override;
+	UE_API virtual TSharedRef<SWidget> GenerateWidgetForColumn(const FName& ColumnId) override;
 
 	/**
 	 * Called when Slate detects that a widget started to be dragged.
@@ -60,20 +62,20 @@ public:
 	 * @param  InMouseEvent  MouseMove that triggered the drag
 	 *
 	 */
-	virtual FReply OnDragDetected(const FGeometry& MyGeometry, const FPointerEvent& MouseEvent) override;
+	UE_API virtual FReply OnDragDetected(const FGeometry& MyGeometry, const FPointerEvent& MouseEvent) override;
 
-	TSharedRef<IToolTip> GetRowToolTip() const;
-	void InvalidateContent();
+	UE_API TSharedRef<IToolTip> GetRowToolTip() const;
+	UE_API void InvalidateContent();
 
 protected:
-	TSharedRef<SWidget> CreateCellWidget(FName ColumnId);
-	FSlateColor GetBackgroundColorAndOpacity() const;
-	FSlateColor GetBackgroundColorAndOpacity(double Time) const;
-	FSlateColor GetOutlineColorAndOpacity() const;
-	const FSlateBrush* GetOutlineBrush(const FName ColumnId) const;
-	bool HandleShouldBeEnabled() const;
-	EVisibility IsColumnVisible(const FName ColumnId) const;
-	void OnSetHoveredCell(TSharedPtr<FTable> InTablePtr, TSharedPtr<FTableColumn> InColumnPtr, FTableTreeNodePtr InTreeNodePtr);
+	UE_API TSharedRef<SWidget> CreateCellWidget(FName ColumnId);
+	UE_API FSlateColor GetBackgroundColorAndOpacity() const;
+	UE_API FSlateColor GetBackgroundColorAndOpacity(double Time) const;
+	UE_API FSlateColor GetOutlineColorAndOpacity() const;
+	UE_API const FSlateBrush* GetOutlineBrush(const FName ColumnId) const;
+	UE_API bool HandleShouldBeEnabled() const;
+	UE_API EVisibility IsColumnVisible(const FName ColumnId) const;
+	UE_API void OnSetHoveredCell(TSharedPtr<FTable> InTablePtr, TSharedPtr<FTableColumn> InColumnPtr, FTableTreeNodePtr InTreeNodePtr);
 
 protected:
 	/** A shared pointer to the table view model. */
@@ -97,3 +99,5 @@ protected:
 };
 
 } // namespace UE::Insights
+
+#undef UE_API
diff --git a/Engine/Source/Developer/TraceInsightsCore/Public/InsightsCore/Table/Widgets/STableTreeViewTooltip.h b/Engine/Source/Developer/TraceInsightsCore/Public/InsightsCore/Table/Widgets/STableTreeViewTooltip.h
--- a/Engine/Source/Developer/TraceInsightsCore/Public/InsightsCore/Table/Widgets/STableTreeViewTooltip.h
+++ b/Engine/Source/Developer/TraceInsightsCore/Public/InsightsCore/Table/Widgets/STableTreeViewTooltip.h
@@ -7,6 +7,8 @@
 #include "Widgets/IToolTip.h"
 #include "Widgets/SToolTip.h"
 
+#define UE_API TRACEINSIGHTSCORE_API
+
 class SGridPanel;
 
 namespace UE::Insights
@@ -19,22 +21,22 @@ class FTableTreeNode;
 ////////////////////////////////////////////////////////////////////////////////////////////////////
 
 /** Tooltip for STableTreeView widget. */
-class TRACEINSIGHTSCORE_API STableTreeViewTooltip
+class STableTreeViewTooltip
 {
 public:
 	STableTreeViewTooltip() = delete;
 
-	static TSharedPtr<SToolTip> GetTableTooltip(const FTable& Table);
-	static TSharedPtr<SToolTip> GetColumnTooltip(const FTableColumn& Column);
-	static TSharedPtr<SToolTip> GetRowTooltip(const TSharedPtr<FTableTreeNode> TreeNodePtr);
+	static UE_API TSharedPtr<SToolTip> GetTableTooltip(const FTable& Table);
+	static UE_API TSharedPtr<SToolTip> GetColumnTooltip(const FTableColumn& Column);
+	static UE_API TSharedPtr<SToolTip> GetRowTooltip(const TSharedPtr<FTableTreeNode> TreeNodePtr);
 
 private:
-	static void AddGridRow(TSharedPtr<SGridPanel> Grid, int32& Row, const FText& Name, const FText& Value);
+	static UE_API void AddGridRow(TSharedPtr<SGridPanel> Grid, int32& Row, const FText& Name, const FText& Value);
 };
 
 ////////////////////////////////////////////////////////////////////////////////////////////////////
 
-class TRACEINSIGHTSCORE_API STableTreeRowToolTip : public IToolTip
+class STableTreeRowToolTip : public IToolTip
 {
 public:
 	STableTreeRowToolTip(const TSharedPtr<FTableTreeNode> InTreeNodePtr) : TreeNodePtr(InTreeNodePtr) {}
@@ -85,3 +87,5 @@ private:
 ////////////////////////////////////////////////////////////////////////////////////////////////////
 
 } // namespace UE::Insights
+
+#undef UE_API
diff --git a/Engine/Source/Developer/TraceInsightsCore/Public/InsightsCore/Table/Widgets/SUntypedDiffTableTreeView.h b/Engine/Source/Developer/TraceInsightsCore/Public/InsightsCore/Table/Widgets/SUntypedDiffTableTreeView.h
--- a/Engine/Source/Developer/TraceInsightsCore/Public/InsightsCore/Table/Widgets/SUntypedDiffTableTreeView.h
+++ b/Engine/Source/Developer/TraceInsightsCore/Public/InsightsCore/Table/Widgets/SUntypedDiffTableTreeView.h
@@ -12,24 +12,26 @@
 #include "InsightsCore/Table/ViewModels/UntypedTable.h"
 #include "InsightsCore/Table/Widgets/SUntypedTableTreeView.h"
 
+#define UE_API TRACEINSIGHTSCORE_API
+
 namespace UE::Insights
 {
 
 ////////////////////////////////////////////////////////////////////////////////////////////////////
 
-class TRACEINSIGHTSCORE_API SUntypedDiffTableTreeView : public SUntypedTableTreeView
+class SUntypedDiffTableTreeView : public SUntypedTableTreeView
 {
 public:
-	void UpdateSourceTableA(const FString& Name, TSharedPtr<TraceServices::IUntypedTable> SourceTable);
-	void UpdateSourceTableB(const FString& Name, TSharedPtr<TraceServices::IUntypedTable> SourceTable);
+	UE_API void UpdateSourceTableA(const FString& Name, TSharedPtr<TraceServices::IUntypedTable> SourceTable);
+	UE_API void UpdateSourceTableB(const FString& Name, TSharedPtr<TraceServices::IUntypedTable> SourceTable);
 
 protected:
-	FReply SwapTables_OnClicked();
-	FText GetSwapButtonText() const;
+	UE_API FReply SwapTables_OnClicked();
+	UE_API FText GetSwapButtonText() const;
 
-	virtual TSharedPtr<SWidget> ConstructToolbar() override;
+	UE_API virtual TSharedPtr<SWidget> ConstructToolbar() override;
 
-	void RequestMergeTables();
+	UE_API void RequestMergeTables();
 
 private:
 	TSharedPtr<TraceServices::IUntypedTable> SourceTableA;
@@ -41,3 +43,5 @@ private:
 ////////////////////////////////////////////////////////////////////////////////////////////////////
 
 } // namespace UE::Insights
+
+#undef UE_API
diff --git a/Engine/Source/Developer/TraceInsightsCore/Public/InsightsCore/Table/Widgets/SUntypedTableTreeView.h b/Engine/Source/Developer/TraceInsightsCore/Public/InsightsCore/Table/Widgets/SUntypedTableTreeView.h
--- a/Engine/Source/Developer/TraceInsightsCore/Public/InsightsCore/Table/Widgets/SUntypedTableTreeView.h
+++ b/Engine/Source/Developer/TraceInsightsCore/Public/InsightsCore/Table/Widgets/SUntypedTableTreeView.h
@@ -13,6 +13,8 @@
 #include "InsightsCore/Table/ViewModels/UntypedTable.h"
 #include "InsightsCore/Table/Widgets/STableTreeView.h"
 
+#define UE_API TRACEINSIGHTSCORE_API
+
 namespace TraceServices
 {
 	class IUntypedTable;
@@ -23,11 +25,11 @@ namespace UE::Insights
 
 ////////////////////////////////////////////////////////////////////////////////////////////////////
 
-class TRACEINSIGHTSCORE_API SUntypedTableTreeView : public STableTreeView
+class SUntypedTableTreeView : public STableTreeView
 {
 public:
-	SUntypedTableTreeView();
-	virtual ~SUntypedTableTreeView();
+	UE_API SUntypedTableTreeView();
+	UE_API virtual ~SUntypedTableTreeView();
 
 	SLATE_BEGIN_ARGS(SUntypedTableTreeView)
 		: _RunInAsyncMode(false)
@@ -39,33 +41,33 @@ public:
 	 * Construct this widget
 	 * @param InArgs - The declaration data for this widget
 	 */
-	void Construct(const FArguments& InArgs, TSharedPtr<FUntypedTable> InTablePtr);
+	UE_API void Construct(const FArguments& InArgs, TSharedPtr<FUntypedTable> InTablePtr);
 
 	TSharedPtr<FUntypedTable> GetUntypedTable() const { return StaticCastSharedPtr<FUntypedTable>(GetTable()); }
 
-	void UpdateSourceTable(TSharedPtr<TraceServices::IUntypedTable> SourceTable);
+	UE_API void UpdateSourceTable(TSharedPtr<TraceServices::IUntypedTable> SourceTable);
 
-	virtual void Reset();
+	UE_API virtual void Reset();
 
 	//////////////////////////////////////////////////
 	// IAsyncOperationStatusProvider
 
-	virtual bool IsRunning() const override;
-	virtual double GetAllOperationsDuration() override;
+	UE_API virtual bool IsRunning() const override;
+	UE_API virtual double GetAllOperationsDuration() override;
 	virtual double GetCurrentOperationDuration() override { return 0.0; }
 	virtual uint32 GetOperationCount() const override { return 1; }
-	virtual FText GetCurrentOperationName() const override;
+	UE_API virtual FText GetCurrentOperationName() const override;
 
 	//////////////////////////////////////////////////
 
-	void SetCurrentOperationNameOverride(const FText& InOperationName);
-	void ClearCurrentOperationNameOverride();
+	UE_API void SetCurrentOperationNameOverride(const FText& InOperationName);
+	UE_API void ClearCurrentOperationNameOverride();
 
 	/**
 	 * Rebuilds the tree (if necessary).
 	 * @param bResync - If true, it forces a resync even if the list did not changed since last sync.
 	 */
-	virtual void RebuildTree(bool bResync);
+	UE_API virtual void RebuildTree(bool bResync);
 
 private:
 	FStopwatch CurrentOperationStopwatch;
@@ -75,3 +77,5 @@ private:
 ////////////////////////////////////////////////////////////////////////////////////////////////////
 
 } // namespace UE::Insights
+
+#undef UE_API
diff --git a/Engine/Source/Developer/TraceInsightsFrontend/Private/InsightsFrontend/Common/InsightsAutomationController.h b/Engine/Source/Developer/TraceInsightsFrontend/Private/InsightsFrontend/Common/InsightsAutomationController.h
--- a/Engine/Source/Developer/TraceInsightsFrontend/Private/InsightsFrontend/Common/InsightsAutomationController.h
+++ b/Engine/Source/Developer/TraceInsightsFrontend/Private/InsightsFrontend/Common/InsightsAutomationController.h
@@ -8,12 +8,14 @@
 #include "Logging/LogMacros.h"
 #include "Containers/Ticker.h"
 
+#define UE_API TRACEINSIGHTSFRONTEND_API
+
 DECLARE_LOG_CATEGORY_EXTERN(InsightsAutomationController, Log, All);
 
 namespace UE::Insights
 {
 
-class TRACEINSIGHTSFRONTEND_API FInsightsAutomationController : public TSharedFromThis<FInsightsAutomationController>
+class FInsightsAutomationController : public TSharedFromThis<FInsightsAutomationController>
 {
 	enum class ETestsState
 	{
@@ -23,16 +25,16 @@ class TRACEINSIGHTSFRONTEND_API FInsightsAutomationController : public TSharedFr
 	};
 
 public:
-	virtual ~FInsightsAutomationController();
+	UE_API virtual ~FInsightsAutomationController();
 
-	virtual void Initialize();
+	UE_API virtual void Initialize();
 
-	bool Tick(float DeltaTime);
+	UE_API bool Tick(float DeltaTime);
 
 	void SetAutoQuit(bool InAutoQuit) { bAutoQuit = InAutoQuit; }
 	bool GetAutoQuit() const { return bAutoQuit; }
 
-	void RunTests(const FString& InCmd);
+	UE_API void RunTests(const FString& InCmd);
 
 private:
 	/** The delegate to be invoked when this manager ticks. */
@@ -46,7 +48,9 @@ private:
 	bool bAutoQuit = false;
 	ETestsState RunningTestsState = ETestsState::NotStarted;
 
-	static const TCHAR* AutoQuitMsgOnComplete;
+	static UE_API const TCHAR* AutoQuitMsgOnComplete;
 };
 
-} // namespace UE::Insights
\ No newline at end of file
+} // namespace UE::Insights
+
+#undef UE_API
diff --git a/Engine/Source/Developer/TraceInsightsFrontend/Public/InsightsFrontend/Common/InsightsFrontendStyle.h b/Engine/Source/Developer/TraceInsightsFrontend/Public/InsightsFrontend/Common/InsightsFrontendStyle.h
--- a/Engine/Source/Developer/TraceInsightsFrontend/Public/InsightsFrontend/Common/InsightsFrontendStyle.h
+++ b/Engine/Source/Developer/TraceInsightsFrontend/Public/InsightsFrontend/Common/InsightsFrontendStyle.h
@@ -11,6 +11,8 @@
 #include "Types/SlateEnums.h"
 #include "UObject/NameTypes.h"
 
+#define UE_API TRACEINSIGHTSFRONTEND_API
+
 struct FSlateBrush;
 
 namespace UE::Insights
@@ -51,14 +53,14 @@ public:
 		}
 	}
 
-	class TRACEINSIGHTSFRONTEND_API FStyle : public FSlateStyleSet
+	class FStyle : public FSlateStyleSet
 	{
 	public:
-		FStyle(const FName& InStyleSetName);
+		UE_API FStyle(const FName& InStyleSetName);
 
-		void Initialize();
+		UE_API void Initialize();
 
-		void SyncParentStyles();
+		UE_API void SyncParentStyles();
 
 	public:
 		// Styles inherited from the parent style
@@ -77,3 +79,5 @@ private:
 };
 
 } // namespace UE::Insights
+
+#undef UE_API
diff --git a/Engine/Source/Developer/TraceInsightsFrontend/Public/InsightsFrontend/ITraceInsightsFrontendModule.h b/Engine/Source/Developer/TraceInsightsFrontend/Public/InsightsFrontend/ITraceInsightsFrontendModule.h
--- a/Engine/Source/Developer/TraceInsightsFrontend/Public/InsightsFrontend/ITraceInsightsFrontendModule.h
+++ b/Engine/Source/Developer/TraceInsightsFrontend/Public/InsightsFrontend/ITraceInsightsFrontendModule.h
@@ -34,7 +34,7 @@ class SConnectionWindow;
 } // namespace UE::Insights
 
 /** Interface for TraceInsightsCore module. */
-class TRACEINSIGHTSFRONTEND_API ITraceInsightsFrontendModule : public IModuleInterface
+class ITraceInsightsFrontendModule : public IModuleInterface
 {
 public:
 	virtual ~ITraceInsightsFrontendModule() {}
diff --git a/Engine/Source/Developer/TraceInsightsFrontend/Public/InsightsFrontend/StoreService/TraceServerControl.h b/Engine/Source/Developer/TraceInsightsFrontend/Public/InsightsFrontend/StoreService/TraceServerControl.h
--- a/Engine/Source/Developer/TraceInsightsFrontend/Public/InsightsFrontend/StoreService/TraceServerControl.h
+++ b/Engine/Source/Developer/TraceInsightsFrontend/Public/InsightsFrontend/StoreService/TraceServerControl.h
@@ -11,6 +11,8 @@
 
 #include <atomic>
 
+#define UE_API TRACEINSIGHTSFRONTEND_API
+
 class FMenuBuilder;
 
 namespace UE::Trace { class FStoreClient; }
@@ -18,13 +20,13 @@ namespace UE::Trace { class FStoreClient; }
 namespace UE::Insights
 {
 
-class TRACEINSIGHTSFRONTEND_API FTraceServerControl
+class FTraceServerControl
 {
 public:
-	FTraceServerControl(const TCHAR* Host, uint32 Port = 0, FName StyleSet = NAME_None);
-	~FTraceServerControl();
+	UE_API FTraceServerControl(const TCHAR* Host, uint32 Port = 0, FName StyleSet = NAME_None);
+	UE_API ~FTraceServerControl();
 
-	void MakeMenu(FMenuBuilder& Builder);
+	UE_API void MakeMenu(FMenuBuilder& Builder);
 
 private:
 	enum class EState : uint8
@@ -36,20 +38,20 @@ private:
 		Command
 	};
 
-	bool ChangeState(EState Expected, EState ChangeTo, uint32 Attempts = 1);
+	UE_API bool ChangeState(EState Expected, EState ChangeTo, uint32 Attempts = 1);
 
-	void TriggerStatusUpdate();
-	void UpdateStatus();
-	void ResetStatus();
+	UE_API void TriggerStatusUpdate();
+	UE_API void UpdateStatus();
+	UE_API void ResetStatus();
 
 	bool CanServerBeStarted() const { return !bIsCancelRequested && bIsLocalHost && State.load(std::memory_order_relaxed) == EState::NotConnected; }
 	bool CanServerBeStopped() const { return !bIsCancelRequested && bIsLocalHost && State.load(std::memory_order_relaxed) == EState::Connected; }
 	bool AreControlsEnabled() const { return !bIsCancelRequested && bIsLocalHost && State.load(std::memory_order_relaxed) == EState::Connected; }
 	bool IsSponsored() const { return bSponsored.load(std::memory_order_relaxed); }
 
-	void OnStart_Clicked();
-	void OnStop_Clicked();
-	void OnSponsored_Changed();
+	UE_API void OnStart_Clicked();
+	UE_API void OnStop_Clicked();
+	UE_API void OnSponsored_Changed();
 
 	std::atomic<EState> State = EState::NotConnected;
 
@@ -72,3 +74,5 @@ private:
 };
 
 } // namespace UE::Insights
+
+#undef UE_API
diff --git a/Engine/Source/Developer/TreeMap/Public/STreeMap.h b/Engine/Source/Developer/TreeMap/Public/STreeMap.h
--- a/Engine/Source/Developer/TreeMap/Public/STreeMap.h
+++ b/Engine/Source/Developer/TreeMap/Public/STreeMap.h
@@ -13,6 +13,8 @@
 #include "ITreeMapCustomization.h"
 #include "TreeMapStyle.h"
 
+#define UE_API TREEMAP_API
+
 class FPaintArgs;
 class FSlateWindowElementList;
 class FWidgetPath;
@@ -20,7 +22,7 @@ class FWidgetPath;
 /**
  * Graphical tree map widget with interactive controls
  */
-class TREEMAP_API STreeMap : public SLeafWidget
+class STreeMap : public SLeafWidget
 {
 
 public:
@@ -109,20 +111,20 @@ public:
 	 * @param	InTreeMapNodeData	The node data we'll be visualizing and editing (root node)
 	 * @param	InCustomization		An optional 'customization' of the tree that allows for filtering and displaying the tree map data according to various attributes
 	 */
-	void Construct( const FArguments& InArgs, const TSharedRef<class FTreeMapNodeData>& InTreeMapNodeData, const TSharedPtr< ITreeMapCustomization >& InCustomization );
+	UE_API void Construct( const FArguments& InArgs, const TSharedRef<class FTreeMapNodeData>& InTreeMapNodeData, const TSharedPtr< ITreeMapCustomization >& InCustomization );
 
 	/** SWidget overrides */
-	virtual int32 OnPaint( const FPaintArgs& Args, const FGeometry& AllottedGeometry, const FSlateRect& MyCullingRect, FSlateWindowElementList& OutDrawElements, int32 LayerId, const FWidgetStyle& InWidgetStyle, bool bParentEnabled ) const override;
-	virtual FVector2D ComputeDesiredSize(float LayoutScaleMultiplier) const override;
-	virtual void Tick( const FGeometry& AllottedGeometry, const double InCurrentTime, const float InDeltaTime ) override;
-	virtual FReply OnMouseButtonDown( const FGeometry& InMyGeometry, const FPointerEvent& InMouseEvent ) override;
-	virtual FReply OnMouseButtonUp( const FGeometry& InMyGeometry, const FPointerEvent& InMouseEvent ) override;
-	virtual FReply OnMouseMove( const FGeometry& InMyGeometry, const FPointerEvent& InMouseEvent ) override;
-	virtual void OnMouseLeave( const FPointerEvent& InMouseEvent ) override;
-	virtual FReply OnMouseButtonDoubleClick( const FGeometry& InMyGeometry, const FPointerEvent& InMouseEvent ) override;
-	virtual FReply OnMouseWheel( const FGeometry& MyGeometry, const FPointerEvent& MouseEvent ) override;
-	virtual bool SupportsKeyboardFocus() const override;
-	virtual FReply OnKeyDown( const FGeometry& MyGeometry, const FKeyEvent& InKeyboardEvent ) override;
+	UE_API virtual int32 OnPaint( const FPaintArgs& Args, const FGeometry& AllottedGeometry, const FSlateRect& MyCullingRect, FSlateWindowElementList& OutDrawElements, int32 LayerId, const FWidgetStyle& InWidgetStyle, bool bParentEnabled ) const override;
+	UE_API virtual FVector2D ComputeDesiredSize(float LayoutScaleMultiplier) const override;
+	UE_API virtual void Tick( const FGeometry& AllottedGeometry, const double InCurrentTime, const float InDeltaTime ) override;
+	UE_API virtual FReply OnMouseButtonDown( const FGeometry& InMyGeometry, const FPointerEvent& InMouseEvent ) override;
+	UE_API virtual FReply OnMouseButtonUp( const FGeometry& InMyGeometry, const FPointerEvent& InMouseEvent ) override;
+	UE_API virtual FReply OnMouseMove( const FGeometry& InMyGeometry, const FPointerEvent& InMouseEvent ) override;
+	UE_API virtual void OnMouseLeave( const FPointerEvent& InMouseEvent ) override;
+	UE_API virtual FReply OnMouseButtonDoubleClick( const FGeometry& InMyGeometry, const FPointerEvent& InMouseEvent ) override;
+	UE_API virtual FReply OnMouseWheel( const FGeometry& MyGeometry, const FPointerEvent& MouseEvent ) override;
+	UE_API virtual bool SupportsKeyboardFocus() const override;
+	UE_API virtual FReply OnKeyDown( const FGeometry& MyGeometry, const FKeyEvent& InKeyboardEvent ) override;
 
 	/**
 	 * Sets a new "active root" for the tree.  This is used to "drill down" to child tree nodes or "climb up" back to the root.
@@ -131,16 +133,16 @@ public:
 	 * @param	NewRoot						The tree's new root.  This can be any node within the tree
 	 * @param	bShouldPlayTransition		If enabled, an animation will play
 	 */
-	void SetTreeRoot( const FTreeMapNodeDataRef& NewRoot, const bool bShouldPlayTransition );
+	UE_API void SetTreeRoot( const FTreeMapNodeDataRef& NewRoot, const bool bShouldPlayTransition );
 
 	/** Returns true if it is possible to zoom out */
-	bool CanZoomOut() const;
+	UE_API bool CanZoomOut() const;
 
 	/** Zooms out, setting root to it's parent if possible, returns false if not */
-	bool ZoomOut();
+	UE_API bool ZoomOut();
 
 	/** Returns the current tree root, which is the largest node visible */
-	FTreeMapNodeDataPtr GetTreeRoot() const;
+	UE_API FTreeMapNodeDataPtr GetTreeRoot() const;
 
 	/**
 	 * Refreshes the tree map from its source data.  The current active root will remain as the base of the tree.  Call this
@@ -148,61 +150,61 @@ public:
 	 * 
 	 * @param	bShouldPlayTransition		If enabled, an animation will play
 	 */
-	void RebuildTreeMap( const bool bShouldPlayTransition );
+	UE_API void RebuildTreeMap( const bool bShouldPlayTransition );
 		
 
 protected:
 
 	/** Finds the node visual that's under the cursor */
-	struct FTreeMapNodeVisualInfo* FindNodeVisualUnderCursor( const FGeometry& MyGeometry, const FVector2D& ScreenSpaceCursorPosition );
+	UE_API struct FTreeMapNodeVisualInfo* FindNodeVisualUnderCursor( const FGeometry& MyGeometry, const FVector2D& ScreenSpaceCursorPosition );
 
 	/** Blends between the current visual state and the state before the last navigation transition, by the specified amount */
-	void MakeBlendedNodeVisual( const int32 VisualIndex, const float NavigationAlpha, FTreeMapNodeVisualInfo& OutVisual ) const;
+	UE_API void MakeBlendedNodeVisual( const int32 VisualIndex, const float NavigationAlpha, FTreeMapNodeVisualInfo& OutVisual ) const;
 
 	/** @return True if we're playing a transition after a navigation action (zoom!) */
-	bool IsNavigationTransitionActive() const;
+	UE_API bool IsNavigationTransitionActive() const;
 
 	/** Grabs a snapshot of the current tree so we can restore it if the user performs an Undo action */
-	void TakeUndoSnapshot();
+	UE_API void TakeUndoSnapshot();
 
 	/** Undoes the last action */
-	void Undo();
+	UE_API void Undo();
 
 	/** Redoes the last action */
-	void Redo();
+	UE_API void Redo();
 
 	/** Reparents DroppedNode to NewParentNode (undoable!) */
-	void ReparentNode( const FTreeMapNodeDataRef DroppedNode, const FTreeMapNodeDataRef NewParentNode );
+	UE_API void ReparentNode( const FTreeMapNodeDataRef DroppedNode, const FTreeMapNodeDataRef NewParentNode );
 
 	/** Deletes the node under the mouse cursor, if any (undoable!) */
-	FReply DeleteHoveredNode();
+	UE_API FReply DeleteHoveredNode();
 
 	/** Insert a new node as a child of the node under the cursor */
-	FReply InsertNewNodeAsChildOfHoveredNode( const FGeometry& MyGeometry );
+	UE_API FReply InsertNewNodeAsChildOfHoveredNode( const FGeometry& MyGeometry );
 
 	/** Searches for the specified node in an identical copy of the node tree. */
-	FTreeMapNodeDataPtr FindNodeInCopiedTree( const FTreeMapNodeDataRef& NodeToFind, const FTreeMapNodeDataRef& OriginalNode, const FTreeMapNodeDataRef& CopiedRootNode ) const;
+	UE_API FTreeMapNodeDataPtr FindNodeInCopiedTree( const FTreeMapNodeDataRef& NodeToFind, const FTreeMapNodeDataRef& OriginalNode, const FTreeMapNodeDataRef& CopiedRootNode ) const;
 
 	/** Pops up a box to allow the user to start renaming a node's title (undoable!) */
-	void StartRenamingNode( const FGeometry& MyGeometry, const FTreeMapNodeDataRef& NodeData, const FVector2D& RelativePosition, const bool bIsNewNode );
+	UE_API void StartRenamingNode( const FGeometry& MyGeometry, const FTreeMapNodeDataRef& NodeData, const FVector2D& RelativePosition, const bool bIsNewNode );
 
 	/** Called when the user commits a rename change */
-	void RenamingNode_OnTextCommitted( const FText& NewText, ETextCommit::Type, TSharedRef<FTreeMapNodeData> NodeToRename );
+	UE_API void RenamingNode_OnTextCommitted( const FText& NewText, ETextCommit::Type, TSharedRef<FTreeMapNodeData> NodeToRename );
 
 	/** Stops renaming a node, committing whatever text was entered */
-	void StopRenamingNode();
+	UE_API void StopRenamingNode();
 
 	/** Called to apply the current 'size based on' and 'color based on' settings to the tree */
-	void ApplyVisualizationToNodes( const FTreeMapNodeDataRef& Node );
+	UE_API void ApplyVisualizationToNodes( const FTreeMapNodeDataRef& Node );
 
 	/** Recursively applies the active visualization to all nodes, such as size by attribute, or color by attribute */
-	void ApplyVisualizationToNodesRecursively( const FTreeMapNodeDataRef& Node, const FLinearColor& DefaultColor, const int32 TreeDepth );
+	UE_API void ApplyVisualizationToNodesRecursively( const FTreeMapNodeDataRef& Node, const FLinearColor& DefaultColor, const int32 TreeDepth );
 
 	/** Displays a context menu at the specified location with options for configuring the tree display */
-	void ShowOptionsMenuAt( const FPointerEvent& InMouseEvent );
+	UE_API void ShowOptionsMenuAt( const FPointerEvent& InMouseEvent );
 
 private:
-	void ShowOptionsMenuAtInternal(const FVector2D& ScreenSpacePosition, const FWidgetPath& WidgetPath);
+	UE_API void ShowOptionsMenuAtInternal(const FVector2D& ScreenSpacePosition, const FWidgetPath& WidgetPath);
 
 private:
 
@@ -354,3 +356,5 @@ private:
 	double HighlightPulseStartTime;
 
 };
+
+#undef UE_API
diff --git a/Engine/Source/Developer/TreeMap/Public/TreeMapStyle.h b/Engine/Source/Developer/TreeMap/Public/TreeMapStyle.h
--- a/Engine/Source/Developer/TreeMap/Public/TreeMapStyle.h
+++ b/Engine/Source/Developer/TreeMap/Public/TreeMapStyle.h
@@ -5,22 +5,26 @@
 #include "CoreMinimal.h"
 #include "Styling/SlateStyle.h"
 
+#define UE_API TREEMAP_API
+
 /**
  * Style data for STreeMap
  */
-class TREEMAP_API FTreeMapStyle
+class FTreeMapStyle
 {
 public:
-	static void Initialize();
-	static void Shutdown();
-	static const ISlateStyle& Get();
-	static FName GetStyleSetName();
+	static UE_API void Initialize();
+	static UE_API void Shutdown();
+	static UE_API const ISlateStyle& Get();
+	static UE_API FName GetStyleSetName();
 
 private:
 
-	static TSharedRef< class FSlateStyleSet > Create();
+	static UE_API TSharedRef< class FSlateStyleSet > Create();
 
 private:
 
-	static TSharedPtr< class FSlateStyleSet > StyleInstance;
+	static UE_API TSharedPtr< class FSlateStyleSet > StyleInstance;
 };
+
+#undef UE_API
diff --git a/Engine/Source/Developer/UncontrolledChangelists/Public/UncontrolledChangelist.h b/Engine/Source/Developer/UncontrolledChangelists/Public/UncontrolledChangelist.h
--- a/Engine/Source/Developer/UncontrolledChangelists/Public/UncontrolledChangelist.h
+++ b/Engine/Source/Developer/UncontrolledChangelists/Public/UncontrolledChangelist.h
@@ -9,18 +9,20 @@
 #include "Misc/Guid.h"
 #include "Templates/SharedPointer.h"
 
+#define UE_API UNCONTROLLEDCHANGELISTS_API
+
 class FText;
 
-class UNCONTROLLEDCHANGELISTS_API FUncontrolledChangelist : public TSharedFromThis<FUncontrolledChangelist, ESPMode::ThreadSafe>
+class FUncontrolledChangelist : public TSharedFromThis<FUncontrolledChangelist, ESPMode::ThreadSafe>
 {
 public:
 	static constexpr const TCHAR* GUID_NAME = TEXT("guid");
-	static const FGuid DEFAULT_UNCONTROLLED_CHANGELIST_GUID;
+	static UE_API const FGuid DEFAULT_UNCONTROLLED_CHANGELIST_GUID;
 
 public:
-	FUncontrolledChangelist();
+	UE_API FUncontrolledChangelist();
 	
-	FUncontrolledChangelist(const FGuid& InGuid);
+	UE_API FUncontrolledChangelist(const FGuid& InGuid);
 	
 	bool operator==(const FUncontrolledChangelist& InOther) const
 	{
@@ -51,14 +53,14 @@ public:
 	 * Serialize the Uncontrolled Changelist to a Json Object.
 	 * @param	OutJsonObject	The Json Object used to serialize.
 	 */
-	void Serialize(TSharedRef<class FJsonObject> OutJsonObject) const;
+	UE_API void Serialize(TSharedRef<class FJsonObject> OutJsonObject) const;
 	
 	/**
 	 * Deserialize the Uncontrolled Changelist from a Json Object.
 	 * @param 	InJsonValue 	The Json Object to read from.
 	 * @return 	True if Deserialization succeeded.
 	 */
-	bool Deserialize(const TSharedRef<class FJsonObject> InJsonValue);
+	UE_API bool Deserialize(const TSharedRef<class FJsonObject> InJsonValue);
 
 private:
 	FGuid		Guid;
@@ -66,3 +68,5 @@ private:
 
 typedef TSharedPtr<FUncontrolledChangelist, ESPMode::ThreadSafe> FUncontrolledChangelistPtr;
 typedef TSharedRef<FUncontrolledChangelist, ESPMode::ThreadSafe> FUncontrolledChangelistRef;
+
+#undef UE_API
