[Insights] TraceInsightsCore: Fixed memory leak for SLazyToolTip (circular reference to the parent widget).

diff --git a/Engine/Source/Developer/TraceInsightsCore/Public/InsightsCore/Widgets/SLazyToolTip.h b/Engine/Source/Developer/TraceInsightsCore/Public/InsightsCore/Widgets/SLazyToolTip.h
--- a/Engine/Source/Developer/TraceInsightsCore/Public/InsightsCore/Widgets/SLazyToolTip.h
+++ b/Engine/Source/Developer/TraceInsightsCore/Public/InsightsCore/Widgets/SLazyToolTip.h
@@ -22,7 +22,7 @@ class SLazyToolTip : public SCompoundWidget, public IToolTip
 public:
 	void Construct(const FArguments& InArgs, TSharedPtr<const ILazyToolTipCreator> InCreator)
 	{
-		Creator = InCreator;
+		WeakCreator = InCreator.ToWeakPtr();
 	}
 
 	virtual TSharedRef<SWidget> AsWidget()
@@ -58,12 +58,20 @@ private:
 	{
 		if (!ToolTipWidget.IsValid())
 		{
-			ToolTipWidget = Creator->CreateTooltip();
+			TSharedPtr<const ILazyToolTipCreator> Creator = WeakCreator.Pin();
+			if (Creator.IsValid())
+			{
+				ToolTipWidget = Creator->CreateTooltip();
+			}
+			else
+			{
+				ToolTipWidget = SNew(SToolTip);
+			}
 			check(ToolTipWidget.IsValid());
 		}
 	}
 
 private:
 	TSharedPtr<SToolTip> ToolTipWidget;
-	TSharedPtr<const ILazyToolTipCreator> Creator;
+	TWeakPtr<const ILazyToolTipCreator> WeakCreator;
 };
