[Insights] TraceInsightsCore/Tables: Fixed memory leak affecting the STableTreeView tree/tables.

diff --git a/Engine/Source/Developer/TraceInsightsCore/Private/InsightsCore/Table/Widgets/STableTreeViewCell.cpp b/Engine/Source/Developer/TraceInsightsCore/Private/InsightsCore/Table/Widgets/STableTreeViewCell.cpp
--- a/Engine/Source/Developer/TraceInsightsCore/Private/InsightsCore/Table/Widgets/STableTreeViewCell.cpp
+++ b/Engine/Source/Developer/TraceInsightsCore/Private/InsightsCore/Table/Widgets/STableTreeViewCell.cpp
@@ -28,7 +28,7 @@ BEGIN_SLATE_FUNCTION_BUILD_OPTIMIZATION
 
 void STableTreeViewCell::Construct(const FArguments& InArgs, const TSharedRef<ITableRow>& InTableRow)
 {
-	TableRow = InTableRow;
+	WeakTableRow = InTableRow.ToWeakPtr();
 
 	TablePtr = InArgs._TablePtr;
 	ColumnPtr = InArgs._ColumnPtr;
@@ -64,8 +64,13 @@ TSharedRef<SWidget> STableTreeViewCell::GenerateWidgetForColumn(const FArguments
 
 TSharedRef<SWidget> STableTreeViewCell::GenerateWidgetForNameColumn(const FArguments& InArgs)
 {
-	TSharedPtr<STableTreeViewRow> Row = StaticCastSharedPtr<STableTreeViewRow, ITableRow>(TableRow);
-	TSharedPtr<IToolTip> RowToolTip = Row->GetRowToolTip();
+	TSharedPtr<IToolTip> RowToolTip;
+	TSharedPtr<ITableRow> TableRow = WeakTableRow.Pin();
+	if (TableRow.IsValid())
+	{
+		TSharedPtr<STableTreeViewRow> Row = StaticCastSharedPtr<STableTreeViewRow, ITableRow>(TableRow);
+		RowToolTip = Row->GetRowToolTip();
+	}
 
 	return
 		SNew(SHorizontalBox)
@@ -172,7 +177,8 @@ END_SLATE_FUNCTION_BUILD_OPTIMIZATION
 
 bool STableTreeViewCell::IsSelected() const
 {
-	return TableRow->IsItemSelected();
+	TSharedPtr<ITableRow> TableRow = WeakTableRow.Pin();
+	return TableRow.IsValid() && TableRow->IsItemSelected();
 }
 
 ////////////////////////////////////////////////////////////////////////////////////////////////////
diff --git a/Engine/Source/Developer/TraceInsightsCore/Public/InsightsCore/Table/Widgets/STableTreeViewCell.h b/Engine/Source/Developer/TraceInsightsCore/Public/InsightsCore/Table/Widgets/STableTreeViewCell.h
--- a/Engine/Source/Developer/TraceInsightsCore/Public/InsightsCore/Table/Widgets/STableTreeViewCell.h
+++ b/Engine/Source/Developer/TraceInsightsCore/Public/InsightsCore/Table/Widgets/STableTreeViewCell.h
@@ -134,7 +134,8 @@ protected:
 	UE_API FLinearColor GetShadowColorAndOpacity() const;
 
 protected:
-	TSharedPtr<ITableRow> TableRow;
+	/** A weak pointer to the parent table row widget. */
+	TWeakPtr<ITableRow> WeakTableRow;
 
 	/** A shared pointer to the table view model. */
 	TSharedPtr<FTable> TablePtr;
diff --git a/Engine/Source/Developer/TraceInsightsCore/Public/InsightsCore/Table/Widgets/STableTreeViewTooltip.h b/Engine/Source/Developer/TraceInsightsCore/Public/InsightsCore/Table/Widgets/STableTreeViewTooltip.h
--- a/Engine/Source/Developer/TraceInsightsCore/Public/InsightsCore/Table/Widgets/STableTreeViewTooltip.h
+++ b/Engine/Source/Developer/TraceInsightsCore/Public/InsightsCore/Table/Widgets/STableTreeViewTooltip.h
@@ -3,7 +3,6 @@
 #pragma once
 
 #include "CoreMinimal.h"
-
 #include "Widgets/IToolTip.h"
 #include "Widgets/SToolTip.h"
 
@@ -40,9 +39,9 @@ class STableTreeRowToolTip : public IToolTip
 {
 public:
 	STableTreeRowToolTip(const TSharedPtr<FTableTreeNode> InTreeNodePtr) : TreeNodePtr(InTreeNodePtr) {}
-	virtual ~STableTreeRowToolTip() { }
+	virtual ~STableTreeRowToolTip() {}
 
-	virtual TSharedRef<class SWidget> AsWidget()
+	virtual TSharedRef<SWidget> AsWidget()
 	{
 		CreateToolTipWidget();
 		return ToolTipWidget.ToSharedRef();
@@ -76,6 +75,7 @@ private:
 		if (!ToolTipWidget.IsValid())
 		{
 			ToolTipWidget = STableTreeViewTooltip::GetRowTooltip(TreeNodePtr);
+			check(ToolTipWidget.IsValid());
 		}
 	}
 
