[Insights] Added debug functionality to display memory stats (used physical and used virtual) for the Unreal Insights app in the UI (next to version string). Double clicking on the version string (top right corner) will toggle display between few modes where memory stats are also displayed.

diff --git a/Engine/Source/Developer/TraceInsights/Private/Insights/TraceInsightsModule.cpp b/Engine/Source/Developer/TraceInsights/Private/Insights/TraceInsightsModule.cpp
--- a/Engine/Source/Developer/TraceInsights/Private/Insights/TraceInsightsModule.cpp
+++ b/Engine/Source/Developer/TraceInsights/Private/Insights/TraceInsightsModule.cpp
@@ -262,7 +262,7 @@ void FTraceInsightsModule::CreateSessionViewer(bool bAllowDebugTools)
 {
 	RegisterTabSpawners();
 
-#if !WITH_EDITOR
+#if !WITH_EDITOR || defined(__INTELLISENSE__)
 
 	//////////////////////////////////////////////////
 	// Create the main window.
@@ -318,7 +318,24 @@ void FTraceInsightsModule::CreateSessionViewer(bool bAllowDebugTools)
 	const EOutputCanBeNullptr OutputCanBeNullptr = EOutputCanBeNullptr::Never;
 	TSharedPtr<SWidget> Content = FGlobalTabmanager::Get()->RestoreFrom(PersistentLayout.ToSharedRef(), RootWindow, bEmbedTitleAreaContent, OutputCanBeNullptr);
 
-	RootWindow->SetContent(Content.ToSharedRef());
+	RootWindow->SetContent(
+		SNew(SOverlay)
+
+		+ SOverlay::Slot()
+		.HAlign(HAlign_Fill)
+		.VAlign(VAlign_Fill)
+		[
+			Content.ToSharedRef()
+		]
+
+		+ SOverlay::Slot()
+		.HAlign(HAlign_Right)
+		.VAlign(VAlign_Top)
+		.Padding(0.0f, 15.0f, 4.0f, 0.0f)
+		[
+			VersionWidgetVM->CreateWidget()
+		]
+	);
 	RootWindow->GetOnWindowClosedEvent().AddRaw(this, &FTraceInsightsModule::OnWindowClosedEvent);
 
 	//////////////////////////////////////////////////
diff --git a/Engine/Source/Developer/TraceInsights/Private/Insights/TraceInsightsModule.h b/Engine/Source/Developer/TraceInsights/Private/Insights/TraceInsightsModule.h
--- a/Engine/Source/Developer/TraceInsights/Private/Insights/TraceInsightsModule.h
+++ b/Engine/Source/Developer/TraceInsights/Private/Insights/TraceInsightsModule.h
@@ -7,8 +7,8 @@
 #include "Insights/IUnrealInsightsModule.h"
 #include "Framework/Docking/TabManager.h"
 
-// TraceAnalysis
-#include "Trace/StoreService.h"
+// TraceInsightsCore
+#include "InsightsCore/Common/MiscUtils.h"
 
 namespace Insights
 {
@@ -34,6 +34,10 @@ namespace UE::Insights
 class FTraceInsightsModule : public IUnrealInsightsModule
 {
 public:
+	FTraceInsightsModule()
+		: VersionWidgetVM(MakeShared<FVersionWidget>())
+	{
+	}
 	virtual ~FTraceInsightsModule()
 	{
 	}
@@ -116,6 +120,8 @@ private:
 	static FString UnrealInsightsLayoutIni;
 
 	TArray<TSharedRef<IInsightsComponent>> Components;
+
+	TSharedRef<FVersionWidget> VersionWidgetVM;
 };
 
 } // namespace UE::Insights
diff --git a/Engine/Source/Developer/TraceInsights/Private/Insights/Widgets/SMajorTabWindow.cpp b/Engine/Source/Developer/TraceInsights/Private/Insights/Widgets/SMajorTabWindow.cpp
--- a/Engine/Source/Developer/TraceInsights/Private/Insights/Widgets/SMajorTabWindow.cpp
+++ b/Engine/Source/Developer/TraceInsights/Private/Insights/Widgets/SMajorTabWindow.cpp
@@ -22,9 +22,6 @@
 	#include "Interfaces/IAnalyticsProvider.h"
 #endif // WITH_EDITOR
 
-// TraceInsightsCore
-#include "InsightsCore/Version.h"
-
 // TraceInsights
 #include "Insights/Common/InsightsMenuBuilder.h"
 #include "Insights/InsightsManager.h"
@@ -252,18 +249,6 @@ void SMajorTabWindow::Construct(const FArguments& InArgs, const TSharedRef<SDock
 		]
 #endif
 
-		// Version
-		+ SOverlay::Slot()
-		.HAlign(HAlign_Right)
-		.VAlign(VAlign_Top)
-		.Padding(0.0f, -16.0f, 4.0f, 0.0f)
-		[
-			SNew(STextBlock)
-			.Clipping(EWidgetClipping::ClipToBoundsWithoutIntersecting)
-			.Text(LOCTEXT("UnrealInsightsVersion", UNREAL_INSIGHTS_VERSION_STRING_EX))
-			.ColorAndOpacity(FLinearColor(0.15f, 0.15f, 0.15f, 1.0f))
-		]
-
 		// Overlay slot for the main window area
 		+ SOverlay::Slot()
 		.HAlign(HAlign_Fill)
diff --git a/Engine/Source/Developer/TraceInsights/Private/Insights/Widgets/SSessionInfoWindow.cpp b/Engine/Source/Developer/TraceInsights/Private/Insights/Widgets/SSessionInfoWindow.cpp
--- a/Engine/Source/Developer/TraceInsights/Private/Insights/Widgets/SSessionInfoWindow.cpp
+++ b/Engine/Source/Developer/TraceInsights/Private/Insights/Widgets/SSessionInfoWindow.cpp
@@ -24,7 +24,6 @@
 
 // TraceInsightsCore
 #include "InsightsCore/Common/TimeUtils.h"
-#include "InsightsCore/Version.h"
 
 // TraceInsights
 #include "Insights/Common/InsightsMenuBuilder.h"
@@ -125,18 +124,6 @@ void SSessionInfoWindow::Construct(const FArguments& InArgs, const TSharedRef<SD
 		]
 #endif
 
-		// Version
-		+ SOverlay::Slot()
-		.HAlign(HAlign_Right)
-		.VAlign(VAlign_Top)
-		.Padding(0.0f, -16.0f, 4.0f, 0.0f)
-		[
-			SNew(STextBlock)
-			.Clipping(EWidgetClipping::ClipToBoundsWithoutIntersecting)
-			.Text(LOCTEXT("UnrealInsightsVersion", UNREAL_INSIGHTS_VERSION_STRING_EX))
-			.ColorAndOpacity(FLinearColor(0.15f, 0.15f, 0.15f, 1.0f))
-		]
-
 		// Overlay slot for the main window area
 		+ SOverlay::Slot()
 		.HAlign(HAlign_Fill)
diff --git a/Engine/Source/Developer/TraceInsightsCore/Private/InsightsCore/Common/MiscUtils.cpp b/Engine/Source/Developer/TraceInsightsCore/Private/InsightsCore/Common/MiscUtils.cpp
--- a/Engine/Source/Developer/TraceInsightsCore/Private/InsightsCore/Common/MiscUtils.cpp
+++ b/Engine/Source/Developer/TraceInsightsCore/Private/InsightsCore/Common/MiscUtils.cpp
@@ -5,6 +5,12 @@
 #include "HAL/FileManagerGeneric.h"
 #include "HAL/PlatformProcess.h"
 #include "Misc/FileHelper.h"
+#include "Widgets/Text/STextBlock.h"
+
+// TraceInsightsCore
+#include "InsightsCore/Version.h"
+
+#define LOCTEXT_NAMESPACE "UE::Insights::MiscUtils"
 
 namespace UE::Insights
 {
@@ -43,4 +49,77 @@ bool FMiscUtils::OpenUnrealInsights(const TCHAR* InCmdLine)
 
 ////////////////////////////////////////////////////////////////////////////////////////////////////
 
+FText FVersionWidget::GetWidgetText() const
+{
+	switch (DisplayMode)
+	{
+		case EDisplayMode::Version:
+			return FText::FromString(UNREAL_INSIGHTS_VERSION_STRING_EX);
+
+		case EDisplayMode::VersionAndMem:
+		{
+			FPlatformMemoryStats Stats = FPlatformMemory::GetStats();
+			FNumberFormattingOptions Options;
+			Options.MinimumFractionalDigits = 2;
+			Options.MaximumFractionalDigits = 2;
+			return FText::Format(LOCTEXT("UnrealInsightsVersionFmt1", "Mem: {0}     {1}"),
+				FText::AsMemory(Stats.UsedPhysical, &Options),
+				FText::FromString(UNREAL_INSIGHTS_VERSION_STRING_EX));
+		}
+
+		case EDisplayMode::VersionAndDetailedMem:
+		{
+			FPlatformMemoryStats Stats = FPlatformMemory::GetStats();
+			FNumberFormattingOptions Options;
+			Options.MinimumFractionalDigits = 3;
+			Options.MaximumFractionalDigits = 3;
+			return FText::Format(LOCTEXT("UnrealInsightsVersionTextFmt2", "Used Physical: {0}     Used Virtual: {1}     {2}"),
+				FText::AsMemory(Stats.UsedPhysical, &Options),
+				FText::AsMemory(Stats.UsedVirtual, &Options),
+				FText::FromString(UNREAL_INSIGHTS_VERSION_STRING_EX));
+		}
+
+		default:
+			return FText::GetEmpty();
+	}
+}
+
+////////////////////////////////////////////////////////////////////////////////////////////////////
+
+FSlateColor FVersionWidget::GetColor() const
+{
+	switch (DisplayMode)
+	{
+		case EDisplayMode::VersionAndMem:
+		case EDisplayMode::VersionAndDetailedMem:
+			return FSlateColor(FLinearColor(1.0f, 0.15f, 0.15f, 1.0f));
+
+		case EDisplayMode::Version:
+		default:
+			return FSlateColor(FLinearColor(0.15f, 0.15f, 0.15f, 1.0f));
+	}
+}
+
+////////////////////////////////////////////////////////////////////////////////////////////////////
+
+FReply FVersionWidget::OnDoubleClicked(const FGeometry&, const FPointerEvent&)
+{
+	NextDisplayMode();
+	return FReply::Handled();
+}
+
+////////////////////////////////////////////////////////////////////////////////////////////////////
+
+TSharedRef<SWidget> FVersionWidget::CreateWidget()
+{
+	return SNew(STextBlock)
+		.Text(this, &FVersionWidget::GetWidgetText)
+		.ColorAndOpacity(this, &FVersionWidget::GetColor)
+		.OnDoubleClicked(this, &FVersionWidget::OnDoubleClicked);
+}
+
+////////////////////////////////////////////////////////////////////////////////////////////////////
+
 } // namespace UE::Insights
+
+#undef LOCTEXT_NAMESPACE
diff --git a/Engine/Source/Developer/TraceInsightsCore/Public/InsightsCore/Common/MiscUtils.h b/Engine/Source/Developer/TraceInsightsCore/Public/InsightsCore/Common/MiscUtils.h
--- a/Engine/Source/Developer/TraceInsightsCore/Public/InsightsCore/Common/MiscUtils.h
+++ b/Engine/Source/Developer/TraceInsightsCore/Public/InsightsCore/Common/MiscUtils.h
@@ -4,6 +4,14 @@
 
 #include "CoreTypes.h"
 
+#include "Input/Events.h"
+#include "Input/Reply.h"
+#include "Internationalization/Text.h"
+#include "Layout/Geometry.h"
+#include "Styling/SlateColor.h"
+
+#include "Templates/SharedPointer.h"
+
 #define UE_API TRACEINSIGHTSCORE_API
 
 namespace UE::Insights
@@ -20,6 +28,44 @@ public:
 	static UE_API bool OpenUnrealInsights(const TCHAR* InCmdLine = nullptr);
 };
 
+class FVersionWidget : public TSharedFromThis<FVersionWidget>
+{
+public:
+	enum class EDisplayMode
+	{
+		Version,
+		VersionAndMem,
+		VersionAndDetailedMem,
+
+		Count
+	};
+
+public:
+	FVersionWidget() {}
+	virtual ~FVersionWidget() {}
+
+	UE_API virtual FText GetWidgetText() const;
+	UE_API virtual FSlateColor GetColor() const;
+	UE_API virtual FReply OnDoubleClicked(const FGeometry&, const FPointerEvent&);
+	UE_API virtual TSharedRef<SWidget> CreateWidget();
+
+	EDisplayMode GetDisplayMode() const
+	{
+		return DisplayMode;
+	}
+	void SetDisplayMode(EDisplayMode InDisplayMode)
+	{
+		DisplayMode = InDisplayMode;
+	}
+	void NextDisplayMode()
+	{
+		DisplayMode = EDisplayMode((int32(DisplayMode) + 1) % int32(EDisplayMode::Count));
+	}
+
+private:
+	EDisplayMode DisplayMode = EDisplayMode::Version;
+};
+
 } // namespace UE::Insights
 
 #undef UE_API
diff --git a/Engine/Source/Developer/TraceInsightsFrontend/Private/InsightsFrontend/TraceInsightsFrontendModule.cpp b/Engine/Source/Developer/TraceInsightsFrontend/Private/InsightsFrontend/TraceInsightsFrontendModule.cpp
--- a/Engine/Source/Developer/TraceInsightsFrontend/Private/InsightsFrontend/TraceInsightsFrontendModule.cpp
+++ b/Engine/Source/Developer/TraceInsightsFrontend/Private/InsightsFrontend/TraceInsightsFrontendModule.cpp
@@ -222,7 +222,25 @@ void FTraceInsightsFrontendModule::CreateFrontendWindow(const FCreateFrontendWin
 	// Restore application layout.
 	const EOutputCanBeNullptr OutputCanBeNullptr = EOutputCanBeNullptr::Never;
 	TSharedPtr<SWidget> Content = FGlobalTabmanager::Get()->RestoreFrom(PersistentLayout.ToSharedRef(), RootWindow, bEmbedTitleAreaContent, OutputCanBeNullptr);
-	RootWindow->SetContent(Content.ToSharedRef());
+
+	RootWindow->SetContent(
+		SNew(SOverlay)
+
+		+ SOverlay::Slot()
+		.HAlign(HAlign_Fill)
+		.VAlign(VAlign_Fill)
+		[
+			Content.ToSharedRef()
+		]
+
+		+ SOverlay::Slot()
+		.HAlign(HAlign_Right)
+		.VAlign(VAlign_Top)
+		.Padding(0.0f, 15.0f, 4.0f, 0.0f)
+		[
+			VersionWidgetVM->CreateWidget()
+		]
+	);
 	//RootWindow->GetOnWindowClosedEvent().AddRaw(this, &FTraceInsightsFrontendModule::OnWindowClosedEvent);
 
 	//////////////////////////////////////////////////
diff --git a/Engine/Source/Developer/TraceInsightsFrontend/Private/InsightsFrontend/TraceInsightsFrontendModule.h b/Engine/Source/Developer/TraceInsightsFrontend/Private/InsightsFrontend/TraceInsightsFrontendModule.h
--- a/Engine/Source/Developer/TraceInsightsFrontend/Private/InsightsFrontend/TraceInsightsFrontendModule.h
+++ b/Engine/Source/Developer/TraceInsightsFrontend/Private/InsightsFrontend/TraceInsightsFrontendModule.h
@@ -2,9 +2,14 @@
 
 #pragma once
 
+#include "CoreTypes.h"
+
 #include "Containers/Ticker.h"
 #include "Framework/Docking/TabManager.h"
 
+// TraceInsightsCore
+#include "InsightsCore/Common/MiscUtils.h"
+
 // TraceInsightsFrontend
 #include "InsightsFrontend/InsightsFrontendSettings.h"
 #include "InsightsFrontend/ITraceInsightsFrontendModule.h"
@@ -21,6 +26,10 @@ class SConnectionWindow;
 class FTraceInsightsFrontendModule : public ITraceInsightsFrontendModule
 {
 public:
+	FTraceInsightsFrontendModule()
+		: VersionWidgetVM(MakeShared<FVersionWidget>())
+	{
+	}
 	virtual ~FTraceInsightsFrontendModule()
 	{
 	}
@@ -94,6 +103,8 @@ private:
 
 	/** Handle to the registered OnTick. */
 	FTSTicker::FDelegateHandle OnTickHandle;
+
+	TSharedRef<FVersionWidget> VersionWidgetVM;
 };
 
 } // UE::Insights
diff --git a/Engine/Source/Developer/TraceInsightsFrontend/Private/InsightsFrontend/Widgets/SConnectionWindow.cpp b/Engine/Source/Developer/TraceInsightsFrontend/Private/InsightsFrontend/Widgets/SConnectionWindow.cpp
--- a/Engine/Source/Developer/TraceInsightsFrontend/Private/InsightsFrontend/Widgets/SConnectionWindow.cpp
+++ b/Engine/Source/Developer/TraceInsightsFrontend/Private/InsightsFrontend/Widgets/SConnectionWindow.cpp
@@ -20,9 +20,6 @@
 #include "Trace/ControlClient.h"
 #include "Trace/StoreConnection.h"
 
-// TraceInsightsCore
-#include "InsightsCore/Version.h"
-
 // TraceInsightsFrontend
 #include "InsightsFrontend/Common/InsightsFrontendStyle.h"
 #include "InsightsFrontend/Common/Log.h"
@@ -64,18 +61,6 @@ void SConnectionWindow::Construct(const FArguments& InArgs, TSharedRef<UE::Trace
 	[
 		SNew(SOverlay)
 
-		// Version
-		+ SOverlay::Slot()
-		.HAlign(HAlign_Right)
-		.VAlign(VAlign_Top)
-		.Padding(0.0f, -16.0f, 4.0f, 0.0f)
-		[
-			SNew(STextBlock)
-			.Clipping(EWidgetClipping::ClipToBoundsWithoutIntersecting)
-			.Text(FText::FromString(TEXT(UNREAL_INSIGHTS_VERSION_STRING_EX)))
-			.ColorAndOpacity(FLinearColor(0.15f, 0.15f, 0.15f, 1.0f))
-		]
-
 		+ SOverlay::Slot()
 		.HAlign(HAlign_Fill)
 		.VAlign(VAlign_Fill)
@@ -97,7 +82,7 @@ void SConnectionWindow::Construct(const FArguments& InArgs, TSharedRef<UE::Trace
 		.HAlign(HAlign_Fill)
 		.VAlign(VAlign_Fill)
 		[
-			SAssignNew(MainContentPanel, SVerticalBox)
+			SNew(SVerticalBox)
 
 			+ SVerticalBox::Slot()
 			.HAlign(HAlign_Fill)
diff --git a/Engine/Source/Developer/TraceInsightsFrontend/Private/InsightsFrontend/Widgets/SConnectionWindow.h b/Engine/Source/Developer/TraceInsightsFrontend/Private/InsightsFrontend/Widgets/SConnectionWindow.h
--- a/Engine/Source/Developer/TraceInsightsFrontend/Private/InsightsFrontend/Widgets/SConnectionWindow.h
+++ b/Engine/Source/Developer/TraceInsightsFrontend/Private/InsightsFrontend/Widgets/SConnectionWindow.h
@@ -43,7 +43,6 @@ private:
 private:
 	TSharedPtr<UE::Trace::FStoreConnection> TraceStoreConnection;
 
-	TSharedPtr<SVerticalBox> MainContentPanel;
 	TSharedPtr<SEditableTextBox> TraceRecorderAddressTextBox;
 	TSharedPtr<SEditableTextBox> RunningInstanceAddressTextBox;
 	TSharedPtr<SEditableTextBox> ChannelsTextBox;
diff --git a/Engine/Source/Developer/TraceInsightsFrontend/Private/InsightsFrontend/Widgets/STraceStoreWindow.cpp b/Engine/Source/Developer/TraceInsightsFrontend/Private/InsightsFrontend/Widgets/STraceStoreWindow.cpp
--- a/Engine/Source/Developer/TraceInsightsFrontend/Private/InsightsFrontend/Widgets/STraceStoreWindow.cpp
+++ b/Engine/Source/Developer/TraceInsightsFrontend/Private/InsightsFrontend/Widgets/STraceStoreWindow.cpp
@@ -46,7 +46,6 @@
 #include "InsightsCore/Common/MiscUtils.h"
 #include "InsightsCore/Common/Stopwatch.h"
 #include "InsightsCore/Table/ViewModels/TableImporter.h"
-#include "InsightsCore/Version.h"
 
 // TraceInsightsFrontend
 #include "InsightsFrontend/Common/InsightsFrontendStyle.h"
@@ -119,18 +118,6 @@ void STraceStoreWindow::Construct(const FArguments& InArgs, TSharedRef<UE::Trace
 	[
 		SNew(SOverlay)
 
-		// Version
-		+ SOverlay::Slot()
-		.HAlign(HAlign_Right)
-		.VAlign(VAlign_Top)
-		.Padding(0.0f, -16.0f, 4.0f, 0.0f)
-		[
-			SNew(STextBlock)
-			.Clipping(EWidgetClipping::ClipToBoundsWithoutIntersecting)
-			.Text(FText::FromString(TEXT(UNREAL_INSIGHTS_VERSION_STRING_EX)))
-			.ColorAndOpacity(FLinearColor(0.15f, 0.15f, 0.15f, 1.0f))
-		]
-
 		+ SOverlay::Slot()
 		.HAlign(HAlign_Fill)
 		.VAlign(VAlign_Fill)
