Fix Insights crash when removing a monitored directory.

diff --git a/Engine/Source/Developer/ToolWidgets/Private/Dialog/DialogCommands.cpp b/Engine/Source/Developer/ToolWidgets/Private/Dialog/DialogCommands.cpp
--- a/Engine/Source/Developer/ToolWidgets/Private/Dialog/DialogCommands.cpp
+++ b/Engine/Source/Developer/ToolWidgets/Private/Dialog/DialogCommands.cpp
@@ -13,7 +13,7 @@ FDialogCommands::FDialogCommands()
 	(
 		TEXT("Dialogs"), // Context name for fast lookup
 		NSLOCTEXT("Contexts", "DialogCommands", "Dialog Commands"), // Localized context name for displaying
-		TEXT("MainFrame"),
+		NAME_None,
 		FAppStyle::GetAppStyleSetName() // Icon Style Set
 	)
 {
diff --git a/Engine/Source/Developer/TraceInsightsCore/Private/InsightsCore/Common/MessageDialogUtils.cpp b/Engine/Source/Developer/TraceInsightsCore/Private/InsightsCore/Common/MessageDialogUtils.cpp
--- a/Engine/Source/Developer/TraceInsightsCore/Private/InsightsCore/Common/MessageDialogUtils.cpp
+++ b/Engine/Source/Developer/TraceInsightsCore/Private/InsightsCore/Common/MessageDialogUtils.cpp
@@ -2,6 +2,7 @@
 
 #include "InsightsCore/Common/MessageDialogUtils.h"
 
+#include "Dialog/DialogCommands.h"
 #include "Dialog/SCustomDialog.h"
 #include "Widgets/Text/STextBlock.h"
 
@@ -14,6 +15,10 @@ namespace UE::Insights
 
 EDialogResponse FMessageDialogUtils::ShowChoiceDialog(const FText& Title, const FText& Content)
 {
+	if (!FDialogCommands::IsRegistered())
+	{
+		FDialogCommands::Register();
+	}
 	TSharedRef<SCustomDialog> Dialog = SNew(SCustomDialog)
 		.Title(Title)
 		.Content()
