[Insights] Fixed crash in Import Table / Diff Tables, when clicking on the grouping breadcrumb or using the filtering configurator before the table is loaded.

diff --git a/Engine/Source/Developer/TraceInsightsCore/Private/InsightsCore/Table/Widgets/STableTreeView.cpp b/Engine/Source/Developer/TraceInsightsCore/Private/InsightsCore/Table/Widgets/STableTreeView.cpp
--- a/Engine/Source/Developer/TraceInsightsCore/Private/InsightsCore/Table/Widgets/STableTreeView.cpp
+++ b/Engine/Source/Developer/TraceInsightsCore/Private/InsightsCore/Table/Widgets/STableTreeView.cpp
@@ -358,13 +358,21 @@ TSharedRef<SWidget> STableTreeView::ConstructHierarchyBreadcrumbTrail()
 			.ToolTipText(LOCTEXT("GroupingsTooltip", "Groupings Hierarchy"))
 			.OnGetMenuContent_Lambda([this]() -> TSharedRef<SWidget>
 				{
-					check(CurrentGroupings.Num() > 0);
-					return GetGroupingCrumbMenuContent(CurrentGroupings.Last());
+					if (CurrentGroupings.Num() > 0)
+					{
+						return GetGroupingCrumbMenuContent(CurrentGroupings.Last());
+					}
+					else
+					{
+						return SNew(SBox);
+					}
 				})
 			.OnComboBoxOpened_Lambda([this]()
 				{
-					check(CurrentGroupings.Num() > 0);
-					OnGroupingCrumbClicked(CurrentGroupings.Last());
+					if (CurrentGroupings.Num() > 0)
+					{
+						OnGroupingCrumbClicked(CurrentGroupings.Last());
+					}
 				})
 			.ButtonContent()
 			[
@@ -1272,6 +1280,10 @@ TSharedRef<SWidget> STableTreeView::ConstructFilterConfiguratorButton()
 	return SNew(SButton)
 		.ButtonStyle(FAppStyle::Get(), "SimpleButton")
 		.ToolTipText(LOCTEXT("FilterConfiguratorBtn_ToolTip", "Opens the filter configurator window."))
+		.IsEnabled_Lambda([this]()
+			{
+				return Table.IsValid() && Table->IsValid();
+			})
 		.OnClicked(this, &STableTreeView::FilterConfigurator_OnClicked)
 		[
 			SNew(SImage)
