Small increase to the time resolution timing shown in insights (without this, values near the low end of the automatic range have almost no precision - e.g. appearing as 1.2s)

#Synced-CL 49852634

diff --git a/Engine/Source/Developer/TraceInsightsCore/Public/InsightsCore/Common/TimeUtils.h b/Engine/Source/Developer/TraceInsightsCore/Public/InsightsCore/Common/TimeUtils.h
--- a/Engine/Source/Developer/TraceInsightsCore/Public/InsightsCore/Common/TimeUtils.h
+++ b/Engine/Source/Developer/TraceInsightsCore/Public/InsightsCore/Common/TimeUtils.h
@@ -65,9 +65,9 @@ struct FTimeSplit
 	bool bIsNaN;
 };
 
-TRACEINSIGHTSCORE_API FString FormatTimeValue(const double Duration, const int32 NumDigits = 1);
-TRACEINSIGHTSCORE_API FString FormatTimeAuto(const double Duration, const int32 NumDigits = 1);
-TRACEINSIGHTSCORE_API FString FormatTimeMs(const double Duration, const int32 NumDigits = 2, bool bAddTimeUnit = false);
+TRACEINSIGHTSCORE_API FString FormatTimeValue(const double Duration, const int32 NumDigits = 2);
+TRACEINSIGHTSCORE_API FString FormatTimeAuto(const double Duration, const int32 NumDigits = 2);
+TRACEINSIGHTSCORE_API FString FormatTimeMs(const double Duration, const int32 NumDigits = 3, bool bAddTimeUnit = false);
 TRACEINSIGHTSCORE_API FString FormatTime(const double Time, const double Precision = 0.0);
 TRACEINSIGHTSCORE_API FString FormatTimeHMS(const double Time, const double Precision = 0.0);
 TRACEINSIGHTSCORE_API void SplitTime(const double Time, FTimeSplit& OutTimeSplit);
