[Insights] Updated GetValueAsSerializableString() for double and float columns to support "+inf", "-inf" and "nan".

#Synced-CL 50439474

diff --git a/Engine/Source/Developer/TraceInsightsCore/Private/InsightsCore/Table/ViewModels/TableColumn.cpp b/Engine/Source/Developer/TraceInsightsCore/Private/InsightsCore/Table/ViewModels/TableColumn.cpp
--- a/Engine/Source/Developer/TraceInsightsCore/Private/InsightsCore/Table/ViewModels/TableColumn.cpp
+++ b/Engine/Source/Developer/TraceInsightsCore/Private/InsightsCore/Table/ViewModels/TableColumn.cpp
@@ -6,6 +6,9 @@
 #include "InsightsCore/Table/ViewModels/TableCellValueGetter.h"
 #include "InsightsCore/Table/ViewModels/TableCellValueSorter.h"
 
+#include <limits>
+#include <cmath>
+
 #define LOCTEXT_NAMESPACE "UE::Insights::FTableColumn"
 
 namespace UE::Insights
@@ -110,20 +113,60 @@ FString FTableColumn::GetValueAsSerializableString(const FBaseTreeNode& InNode)
 	{
 		case ETableCellDataType::Double:
 		{
-			FString Str = FString::Printf(TEXT("%.9f"), Value->AsDouble());
-			TrimEndingZero(Str);
-			return Str;
+			const double DblValue = Value->AsDouble();
+			if (DblValue == 0.0)
+			{
+				return FString("0");
+			}
+			else if (DblValue == +DBL_MAX || DblValue == +std::numeric_limits<double>::infinity())
+			{
+				return FString("+inf");
+			}
+			else if (DblValue == -DBL_MAX || DblValue == -std::numeric_limits<double>::infinity())
+			{
+				return FString("-inf");
+			}
+			else if (std::isnan(DblValue))
+			{
+				return FString("nan");
+			}
+			else
+			{
+				FString Str = FString::Printf(TEXT("%.9f"), DblValue);
+				TrimEndingZero(Str);
+				return Str;
+			}
 		}
 		case ETableCellDataType::Float:
 		{
-			FString Str = FString::Printf(TEXT("%.6f"), Value->AsFloat());
-			TrimEndingZero(Str);
-			return Str;
+			const float FltValue = Value->AsFloat();
+			if (FltValue == 0.0f)
+			{
+				return FString("0");
+			}
+			else if (FltValue == +FLT_MAX || FltValue == +std::numeric_limits<float>::infinity())
+			{
+				return FString("+inf");
+			}
+			else if (FltValue == -FLT_MAX || FltValue == -std::numeric_limits<float>::infinity())
+			{
+				return FString("-inf");
+			}
+			else if (std::isnan(FltValue))
+			{
+				return FString("nan");
+			}
+			else
+			{
+				FString Str = FString::Printf(TEXT("%.6f"), FltValue);
+				TrimEndingZero(Str);
+				return Str;
+			}
 		}
 		case ETableCellDataType::Int64:
 			return FString::Printf(TEXT("%lli"), Value->AsInt64());
 		case ETableCellDataType::Bool:
-			return Value->AsBool() ? TEXT("true") : TEXT("false");
+			return Value->AsBool() ? FString("true") : FString("false");
 		case ETableCellDataType::CString:
 			return Value->AsString();
 		case ETableCellDataType::Text:
