[Insights Style] Fixed clipping hamburger menu on Timing View Toolbar

#rnx
#Synced-CL 50485454

diff --git a/Engine/Source/Developer/TraceInsightsCore/Private/InsightsCore/Common/InsightsCoreStyle.cpp b/Engine/Source/Developer/TraceInsightsCore/Private/InsightsCore/Common/InsightsCoreStyle.cpp
--- a/Engine/Source/Developer/TraceInsightsCore/Private/InsightsCore/Common/InsightsCoreStyle.cpp
+++ b/Engine/Source/Developer/TraceInsightsCore/Private/InsightsCore/Common/InsightsCoreStyle.cpp
@@ -290,6 +290,8 @@ void FInsightsCoreStyle::FStyle::Initialize()
 		SecondaryToolbarStyle.ButtonStyle.SetNormalPadding( FMargin(6.0f, 2.0f, 4.0f, 2.0f));
 		SecondaryToolbarStyle.ButtonStyle.SetPressedPadding(FMargin(6.0f, 2.0f, 4.0f, 2.0f));
 
+		SecondaryToolbarStyle.WrapButtonStyle.SetWrapButtonPadding(FMargin(0.0f, -2.0f, 0.0f, -2.0f));
+
 		//SecondaryToolbarStyle.IconSize.Set(16.0f, 16.0f);
 
 		Set("SecondaryToolbar", SecondaryToolbarStyle);
