[Insights] TraceInsightsCore: Added constructors without Name to FBaseTreeNode and FTableTreeNode classes (as a first step toward deprecating the FName usage in FBaseTreeNode/FTableTreeNode). Also added explicit SetDisplayName() to FCustomTableTreeNode.

#Synced-CL 51157224

diff --git a/Engine/Source/Developer/TraceInsightsCore/Private/InsightsCore/Table/ViewModels/BaseTreeNode.cpp b/Engine/Source/Developer/TraceInsightsCore/Private/InsightsCore/Table/ViewModels/BaseTreeNode.cpp
--- a/Engine/Source/Developer/TraceInsightsCore/Private/InsightsCore/Table/ViewModels/BaseTreeNode.cpp
+++ b/Engine/Source/Developer/TraceInsightsCore/Private/InsightsCore/Table/ViewModels/BaseTreeNode.cpp
@@ -21,7 +21,7 @@ FBaseTreeNode::FGroupNodeData FBaseTreeNode::DefaultGroupData;
 
 const FText FBaseTreeNode::GetDisplayName() const
 {
-	return FText::FromName(GetName());
+	return FText::FromName(Name);
 }
 
 ////////////////////////////////////////////////////////////////////////////////////////////////////
diff --git a/Engine/Source/Developer/TraceInsightsCore/Public/InsightsCore/Table/ViewModels/BaseTreeNode.h b/Engine/Source/Developer/TraceInsightsCore/Public/InsightsCore/Table/ViewModels/BaseTreeNode.h
--- a/Engine/Source/Developer/TraceInsightsCore/Public/InsightsCore/Table/ViewModels/BaseTreeNode.h
+++ b/Engine/Source/Developer/TraceInsightsCore/Public/InsightsCore/Table/ViewModels/BaseTreeNode.h
@@ -67,15 +67,24 @@ protected:
 	};
 
 public:
+	/** Initialization constructor for the node. */
+	explicit FBaseTreeNode(bool bInIsGroup)
+		: GroupData(bInIsGroup ? new FGroupNodeData() : &DefaultGroupData)
+	{
+	}
+
+	//////////////////////////////////////////////////
+	// Deprecated constructor
+
 	/** Initialization constructor for the node. */
 	explicit FBaseTreeNode(const FName InName, bool bInIsGroup)
-		: DefaultSortOrder(0)
-		, Name(InName)
-		, Parent(nullptr)
+		: Name(InName)
 		, GroupData(bInIsGroup ? new FGroupNodeData() : &DefaultGroupData)
 	{
 	}
 
+	//////////////////////////////////////////////////
+
 	virtual ~FBaseTreeNode()
 	{
 		RemoveGroupData();
@@ -515,10 +524,10 @@ protected:
 
 private:
 	/** The default sort order. Index used to optimize sorting. */
-	int32 DefaultSortOrder;
+	int32 DefaultSortOrder = 0;
 
 	/** The name of this node. */
-	const FName Name;
+	const FName Name; // deprecated
 
 	/** A weak pointer to the parent group of this node. */
 	FBaseTreeNodeWeak Parent;
diff --git a/Engine/Source/Developer/TraceInsightsCore/Public/InsightsCore/Table/ViewModels/TableTreeNode.h b/Engine/Source/Developer/TraceInsightsCore/Public/InsightsCore/Table/ViewModels/TableTreeNode.h
--- a/Engine/Source/Developer/TraceInsightsCore/Public/InsightsCore/Table/ViewModels/TableTreeNode.h
+++ b/Engine/Source/Developer/TraceInsightsCore/Public/InsightsCore/Table/ViewModels/TableTreeNode.h
@@ -53,6 +53,36 @@ class FTableTreeNode : public FBaseTreeNode
 	INSIGHTS_DECLARE_RTTI(FTableTreeNode, FBaseTreeNode, UE_API)
 
 public:
+	/** Initialization constructor for a table record node. */
+	explicit FTableTreeNode(TWeakPtr<FTable> InParentTable, int32 InRowIndex)
+		: FBaseTreeNode(false)
+		, ParentTable(InParentTable)
+		, RowId(InRowIndex)
+		, AggregatedValues(nullptr)
+	{
+	}
+
+	/** Initialization constructor for a group node. */
+	explicit FTableTreeNode(TWeakPtr<FTable> InParentTable)
+		: FBaseTreeNode(true)
+		, ParentTable(InParentTable)
+		, RowId(FTableRowId::InvalidRowIndex)
+		, AggregatedValues(nullptr)
+	{
+	}
+
+	/** Initialization constructor for a table record node. */
+	explicit FTableTreeNode(TWeakPtr<FTable> InParentTable, int32 InRowIndex, bool IsGroup)
+		: FBaseTreeNode(IsGroup)
+		, ParentTable(InParentTable)
+		, RowId(InRowIndex)
+		, AggregatedValues(nullptr)
+	{
+	}
+
+	//////////////////////////////////////////////////
+	// Deprecated constructors
+
 	/** Initialization constructor for a table record node. */
 	explicit FTableTreeNode(const FName InName, TWeakPtr<FTable> InParentTable, int32 InRowIndex)
 		: FBaseTreeNode(InName, false)
@@ -80,6 +110,8 @@ public:
 	{
 	}
 
+	//////////////////////////////////////////////////
+
 	virtual ~FTableTreeNode()
 	{
 		CleanupAggregatedValues();
@@ -174,9 +206,58 @@ class FCustomTableTreeNode : public FTableTreeNode
 	INSIGHTS_DECLARE_RTTI(FCustomTableTreeNode, FTableTreeNode, UE_API)
 
 public:
+	/** Initialization constructor for a table record node. */
+	explicit FCustomTableTreeNode(TWeakPtr<FTable> InParentTable, int32 InRowIndex, const FSlateBrush* InIconBrush, FLinearColor InColor, bool IsGroup)
+		: FTableTreeNode(InParentTable, InRowIndex, IsGroup)
+		, IconBrush(InIconBrush)
+		, IconColor(InColor)
+		, Color(InColor)
+	{
+	}
+
+	/** Initialization constructor for a table record node. */
+	explicit FCustomTableTreeNode(TWeakPtr<FTable> InParentTable, int32 InRowIndex, const FSlateBrush* InIconBrush, FLinearColor InIconColor, FLinearColor InColor, bool IsGroup)
+		: FTableTreeNode(InParentTable, InRowIndex, IsGroup)
+		, IconBrush(InIconBrush)
+		, IconColor(InIconColor)
+		, Color(InColor)
+	{
+	}
+
+	/** Initialization constructor for the group node. */
+	explicit FCustomTableTreeNode(TWeakPtr<FTable> InParentTable, FLinearColor InColor)
+		: FTableTreeNode(InParentTable)
+		, IconBrush(FBaseTreeNode::GetDefaultIcon(true))
+		, IconColor(InColor)
+		, Color(InColor)
+	{
+	}
+
+	/** Initialization constructor for the group node. */
+	explicit FCustomTableTreeNode(TWeakPtr<FTable> InParentTable, const FSlateBrush* InIconBrush, FLinearColor InColor)
+		: FTableTreeNode(InParentTable)
+		, IconBrush(InIconBrush)
+		, IconColor(InColor)
+		, Color(InColor)
+	{
+	}
+
+	/** Initialization constructor for the group node. */
+	explicit FCustomTableTreeNode(TWeakPtr<FTable> InParentTable, const FSlateBrush* InIconBrush, FLinearColor InIconColor, FLinearColor InColor)
+		: FTableTreeNode(InParentTable)
+		, IconBrush(InIconBrush)
+		, IconColor(InIconColor)
+		, Color(InColor)
+	{
+	}
+
+	//////////////////////////////////////////////////
+	// Deprecated constructors
+
 	/** Initialization constructor for a table record node. */
 	explicit FCustomTableTreeNode(const FName InName, TWeakPtr<FTable> InParentTable, int32 InRowIndex, const FSlateBrush* InIconBrush, FLinearColor InColor, bool IsGroup)
 		: FTableTreeNode(InName, InParentTable, InRowIndex, IsGroup)
+		, DisplayName(FText::FromName(InName))
 		, IconBrush(InIconBrush)
 		, IconColor(InColor)
 		, Color(InColor)
@@ -186,6 +267,7 @@ public:
 	/** Initialization constructor for a table record node. */
 	explicit FCustomTableTreeNode(const FName InName, TWeakPtr<FTable> InParentTable, int32 InRowIndex, const FSlateBrush* InIconBrush, FLinearColor InIconColor, FLinearColor InColor, bool IsGroup)
 		: FTableTreeNode(InName, InParentTable, InRowIndex, IsGroup)
+		, DisplayName(FText::FromName(InName))
 		, IconBrush(InIconBrush)
 		, IconColor(InIconColor)
 		, Color(InColor)
@@ -195,6 +277,7 @@ public:
 	/** Initialization constructor for the group node. */
 	explicit FCustomTableTreeNode(const FName InName, TWeakPtr<FTable> InParentTable, FLinearColor InColor)
 		: FTableTreeNode(InName, InParentTable)
+		, DisplayName(FText::FromName(InName))
 		, IconBrush(FBaseTreeNode::GetDefaultIcon(true))
 		, IconColor(InColor)
 		, Color(InColor)
@@ -204,6 +287,7 @@ public:
 	/** Initialization constructor for the group node. */
 	explicit FCustomTableTreeNode(const FName InName, TWeakPtr<FTable> InParentTable, const FSlateBrush* InIconBrush, FLinearColor InColor)
 		: FTableTreeNode(InName, InParentTable)
+		, DisplayName(FText::FromName(InName))
 		, IconBrush(InIconBrush)
 		, IconColor(InColor)
 		, Color(InColor)
@@ -213,14 +297,44 @@ public:
 	/** Initialization constructor for the group node. */
 	explicit FCustomTableTreeNode(const FName InName, TWeakPtr<FTable> InParentTable, const FSlateBrush* InIconBrush, FLinearColor InIconColor, FLinearColor InColor)
 		: FTableTreeNode(InName, InParentTable)
+		, DisplayName(FText::FromName(InName))
 		, IconBrush(InIconBrush)
 		, IconColor(InIconColor)
 		, Color(InColor)
 	{
 	}
 
+	//////////////////////////////////////////////////
+
 	virtual ~FCustomTableTreeNode() = default;
 
+	//////////////////////////////////////////////////
+	// Display Name
+
+	virtual const FText GetDisplayName() const override
+	{
+		return DisplayName;
+	}
+
+	/**
+	 * Sets the display name for this node.
+	 */
+	void SetDisplayName(const FText& InDisplayName)
+	{
+		DisplayName = InDisplayName;
+	}
+
+	/**
+	 * Sets the display name for this node.
+	 */
+	void SetDisplayName(FText&& InDisplayName)
+	{
+		DisplayName = InDisplayName;
+	}
+
+	//////////////////////////////////////////////////
+	// Icon
+
 	virtual const FSlateBrush* GetIcon() const override
 	{
 		return IconBrush;
@@ -234,24 +348,51 @@ public:
 		IconBrush = InIconBrush;
 	}
 
+	//////////////////////////////////////////////////
+	// Icon Color
+
 	virtual FLinearColor GetIconColor() const override
 	{
 		return IconColor;
 	}
 
+	/**
+	 * Sets the color tint for the icon of this node.
+	 */
+	void SetIconColor(const FLinearColor& InIconColor)
+	{
+		IconColor = InIconColor;
+	}
+
+	//////////////////////////////////////////////////
+	// Text Color
+
 	virtual FLinearColor GetColor() const override
 	{
 		return Color;
 	}
 
+	/**
+	 * Sets the color tint for the display name of this node.
+	 */
+	void SetColor(const FLinearColor& InColor)
+	{
+		Color = InColor;
+	}
+
+	//////////////////////////////////////////////////
+
 private:
+	/** The display name of this node. */
+	FText DisplayName;
+
 	/** The icon of this node. */
 	const FSlateBrush* IconBrush;
 
 	/** The color tint for the icon of this node. */
 	FLinearColor IconColor;
 
-	/** The color tint for the name text of this node. */
+	/** The color tint for the display name of this node. */
 	FLinearColor Color;
 };
 
