[Insights] TraceInsightsCore: Phase 2 in deprecating FName (and DefaultSortOrder) from FBaseTreeNode:
- Added UE_DEPRECATED(5.8 for deprecated API  in FBaseTreeNode, FTableTreeNode and FCustomTableTreeNode. But commented out for now (will be enabled once all current usages are fixed).
- Made FBaseTreeNode::GetDefaultSortOrder() virtual and deprecated SetDefaultSortOrder().
- Added GetDefaultSortOrder() implementation to FTableTreeNode.
- Moved deprecated constructors for FTableTreeNode and FCustomTableTreeNode to cpp file.
- Added constructors with DisplayName to FCustomTableTreeNode.

#Synced-CL 51163422

diff --git a/Engine/Source/Developer/TraceInsightsCore/Private/InsightsCore/Table/ViewModels/TableTreeNode.cpp b/Engine/Source/Developer/TraceInsightsCore/Private/InsightsCore/Table/ViewModels/TableTreeNode.cpp
--- a/Engine/Source/Developer/TraceInsightsCore/Private/InsightsCore/Table/ViewModels/TableTreeNode.cpp
+++ b/Engine/Source/Developer/TraceInsightsCore/Private/InsightsCore/Table/ViewModels/TableTreeNode.cpp
@@ -10,6 +10,104 @@ namespace UE::Insights
 INSIGHTS_IMPLEMENT_RTTI(FTableTreeNode)
 INSIGHTS_IMPLEMENT_RTTI(FCustomTableTreeNode)
 
+////////////////////////////////////////////////////////////////////////////////////////////////////
+// FTableTreeNode
+////////////////////////////////////////////////////////////////////////////////////////////////////
+
+PRAGMA_DISABLE_DEPRECATION_WARNINGS
+
+FTableTreeNode::FTableTreeNode(const FName InName, TWeakPtr<FTable> InParentTable, int32 InRowIndex)
+	: FBaseTreeNode(InName, false)
+	, ParentTable(InParentTable)
+	, RowId(InRowIndex)
+{
+}
+
+FTableTreeNode::FTableTreeNode(const FName InGroupName, TWeakPtr<FTable> InParentTable)
+	: FBaseTreeNode(InGroupName, true)
+	, ParentTable(InParentTable)
+	, RowId(FTableRowId::InvalidRowIndex)
+{
+}
+
+FTableTreeNode::FTableTreeNode(const FName InName, TWeakPtr<FTable> InParentTable, int32 InRowIndex, bool IsGroup)
+	: FBaseTreeNode(InName, IsGroup)
+	, ParentTable(InParentTable)
+	, RowId(InRowIndex)
+{
+}
+
+PRAGMA_ENABLE_DEPRECATION_WARNINGS
+
+const FText FTableTreeNode::GetDisplayName() const
+{
+PRAGMA_DISABLE_DEPRECATION_WARNINGS
+	if (GetName() != NAME_None)
+	{
+		return FText::FromName(GetName());
+	}
+PRAGMA_ENABLE_DEPRECATION_WARNINGS
+
+	return RowId.HasValidIndex() ?
+		FText::Format(LOCTEXT("TableTreeNode_RowFmt", "row_{0}"), RowId.RowIndex) :
+		LOCTEXT("TableTreeNode_InvalidRow", "invalid_row");
+}
+
+////////////////////////////////////////////////////////////////////////////////////////////////////
+// FCustomTableTreeNode
+////////////////////////////////////////////////////////////////////////////////////////////////////
+
+PRAGMA_DISABLE_DEPRECATION_WARNINGS
+
+FCustomTableTreeNode::FCustomTableTreeNode(const FName InName, TWeakPtr<FTable> InParentTable, int32 InRowIndex, const FSlateBrush* InIconBrush, FLinearColor InColor, bool IsGroup)
+	: FTableTreeNode(InName, InParentTable, InRowIndex, IsGroup)
+	, DisplayName(FText::FromName(InName))
+	, IconBrush(InIconBrush)
+	, IconColor(InColor)
+	, Color(InColor)
+{
+}
+
+FCustomTableTreeNode::FCustomTableTreeNode(const FName InName, TWeakPtr<FTable> InParentTable, int32 InRowIndex, const FSlateBrush* InIconBrush, FLinearColor InIconColor, FLinearColor InColor, bool IsGroup)
+	: FTableTreeNode(InName, InParentTable, InRowIndex, IsGroup)
+	, DisplayName(FText::FromName(InName))
+	, IconBrush(InIconBrush)
+	, IconColor(InIconColor)
+	, Color(InColor)
+{
+}
+
+FCustomTableTreeNode::FCustomTableTreeNode(const FName InName, TWeakPtr<FTable> InParentTable, FLinearColor InColor)
+	: FTableTreeNode(InName, InParentTable)
+	, DisplayName(FText::FromName(InName))
+	, IconBrush(FBaseTreeNode::GetDefaultIcon(true))
+	, IconColor(InColor)
+	, Color(InColor)
+{
+}
+
+FCustomTableTreeNode::FCustomTableTreeNode(const FName InName, TWeakPtr<FTable> InParentTable, const FSlateBrush* InIconBrush, FLinearColor InColor)
+	: FTableTreeNode(InName, InParentTable)
+	, DisplayName(FText::FromName(InName))
+	, IconBrush(InIconBrush)
+	, IconColor(InColor)
+	, Color(InColor)
+{
+}
+
+FCustomTableTreeNode::FCustomTableTreeNode(const FName InName, TWeakPtr<FTable> InParentTable, const FSlateBrush* InIconBrush, FLinearColor InIconColor, FLinearColor InColor)
+	: FTableTreeNode(InName, InParentTable)
+	, DisplayName(FText::FromName(InName))
+	, IconBrush(InIconBrush)
+	, IconColor(InIconColor)
+	, Color(InColor)
+{
+}
+
+PRAGMA_ENABLE_DEPRECATION_WARNINGS
+
+////////////////////////////////////////////////////////////////////////////////////////////////////
+
 } // namespace UE::Insights
 
 #undef LOCTEXT_NAMESPACE
diff --git a/Engine/Source/Developer/TraceInsightsCore/Public/InsightsCore/Table/ViewModels/BaseTreeNode.h b/Engine/Source/Developer/TraceInsightsCore/Public/InsightsCore/Table/ViewModels/BaseTreeNode.h
--- a/Engine/Source/Developer/TraceInsightsCore/Public/InsightsCore/Table/ViewModels/BaseTreeNode.h
+++ b/Engine/Source/Developer/TraceInsightsCore/Public/InsightsCore/Table/ViewModels/BaseTreeNode.h
@@ -67,24 +67,22 @@ protected:
 	};
 
 public:
+	UE_NONCOPYABLE(FBaseTreeNode);
+
 	/** Initialization constructor for the node. */
 	explicit FBaseTreeNode(bool bInIsGroup)
 		: GroupData(bInIsGroup ? new FGroupNodeData() : &DefaultGroupData)
 	{
 	}
 
-	//////////////////////////////////////////////////
-	// Deprecated constructor
-
 	/** Initialization constructor for the node. */
+	//UE_DEPRECATED(5.8, "Use the constructor without FName")
 	explicit FBaseTreeNode(const FName InName, bool bInIsGroup)
 		: Name(InName)
 		, GroupData(bInIsGroup ? new FGroupNodeData() : &DefaultGroupData)
 	{
 	}
 
-	//////////////////////////////////////////////////
-
 	virtual ~FBaseTreeNode()
 	{
 		RemoveGroupData();
@@ -93,6 +91,7 @@ public:
 	/**
 	 * @returns a name of this node.
 	 */
+	//UE_DEPRECATED(5.8, "Override the GetDisplayName() instead")
 	const FName& GetName() const
 	{
 		return Name;
@@ -503,11 +502,14 @@ public:
 		}
 	}
 
-	uint32 GetDefaultSortOrder() const
+	//////////////////////////////////////////////////
+
+	virtual uint32 GetDefaultSortOrder() const
 	{
 		return DefaultSortOrder;
 	}
 
+	//UE_DEPRECATED(5.8, "Override GetDefaultSortOrder() in a custom node instead")
 	void SetDefaultSortOrder(uint32 Order)
 	{
 		DefaultSortOrder = Order;
@@ -524,10 +526,10 @@ protected:
 
 private:
 	/** The default sort order. Index used to optimize sorting. */
-	int32 DefaultSortOrder = 0;
+	uint32 DefaultSortOrder = 0; // deprecated in UE 5.8
 
 	/** The name of this node. */
-	const FName Name; // deprecated
+	const FName Name; // deprecated in UE 5.8
 
 	/** A weak pointer to the parent group of this node. */
 	FBaseTreeNodeWeak Parent;
diff --git a/Engine/Source/Developer/TraceInsightsCore/Public/InsightsCore/Table/ViewModels/TableTreeNode.h b/Engine/Source/Developer/TraceInsightsCore/Public/InsightsCore/Table/ViewModels/TableTreeNode.h
--- a/Engine/Source/Developer/TraceInsightsCore/Public/InsightsCore/Table/ViewModels/TableTreeNode.h
+++ b/Engine/Source/Developer/TraceInsightsCore/Public/InsightsCore/Table/ViewModels/TableTreeNode.h
@@ -53,12 +53,13 @@ class FTableTreeNode : public FBaseTreeNode
 	INSIGHTS_DECLARE_RTTI(FTableTreeNode, FBaseTreeNode, UE_API)
 
 public:
+	UE_NONCOPYABLE(FTableTreeNode);
+
 	/** Initialization constructor for a table record node. */
-	explicit FTableTreeNode(TWeakPtr<FTable> InParentTable, int32 InRowIndex)
-		: FBaseTreeNode(false)
+	explicit FTableTreeNode(TWeakPtr<FTable> InParentTable, int32 InRowIndex, bool IsGroup = false)
+		: FBaseTreeNode(IsGroup)
 		, ParentTable(InParentTable)
 		, RowId(InRowIndex)
-		, AggregatedValues(nullptr)
 	{
 	}
 
@@ -67,16 +68,6 @@ public:
 		: FBaseTreeNode(true)
 		, ParentTable(InParentTable)
 		, RowId(FTableRowId::InvalidRowIndex)
-		, AggregatedValues(nullptr)
-	{
-	}
-
-	/** Initialization constructor for a table record node. */
-	explicit FTableTreeNode(TWeakPtr<FTable> InParentTable, int32 InRowIndex, bool IsGroup)
-		: FBaseTreeNode(IsGroup)
-		, ParentTable(InParentTable)
-		, RowId(InRowIndex)
-		, AggregatedValues(nullptr)
 	{
 	}
 
@@ -84,31 +75,16 @@ public:
 	// Deprecated constructors
 
 	/** Initialization constructor for a table record node. */
-	explicit FTableTreeNode(const FName InName, TWeakPtr<FTable> InParentTable, int32 InRowIndex)
-		: FBaseTreeNode(InName, false)
-		, ParentTable(InParentTable)
-		, RowId(InRowIndex)
-		, AggregatedValues(nullptr)
-	{
-	}
+	//UE_DEPRECATED(5.8, "Use the constructor without FName.")
+	UE_API explicit FTableTreeNode(const FName InName, TWeakPtr<FTable> InParentTable, int32 InRowIndex);
 
 	/** Initialization constructor for a group node. */
-	explicit FTableTreeNode(const FName InGroupName, TWeakPtr<FTable> InParentTable)
-		: FBaseTreeNode(InGroupName, true)
-		, ParentTable(InParentTable)
-		, RowId(FTableRowId::InvalidRowIndex)
-		, AggregatedValues(nullptr)
-	{
-	}
+	//UE_DEPRECATED(5.8, "Use the constructor without FName.")
+	UE_API explicit FTableTreeNode(const FName InGroupName, TWeakPtr<FTable> InParentTable);
 
 	/** Initialization constructor for a table record node. */
-	explicit FTableTreeNode(const FName InName, TWeakPtr<FTable> InParentTable, int32 InRowIndex, bool IsGroup)
-		: FBaseTreeNode(InName, IsGroup)
-		, ParentTable(InParentTable)
-		, RowId(InRowIndex)
-		, AggregatedValues(nullptr)
-	{
-	}
+	//UE_DEPRECATED(5.8, "Use the constructor without FName.")
+	UE_API explicit FTableTreeNode(const FName InName, TWeakPtr<FTable> InParentTable, int32 InRowIndex, bool IsGroup);
 
 	//////////////////////////////////////////////////
 
@@ -121,6 +97,8 @@ public:
 	FTableRowId GetRowId() const { return RowId; }
 	int32 GetRowIndex() const { return RowId.RowIndex; }
 
+	UE_API virtual const FText GetDisplayName() const override;
+
 	//////////////////////////////////////////////////
 	// Aggregation
 
@@ -166,6 +144,7 @@ public:
 
 	const FTableCellValue& GetAggregatedValue(const FName& ColumnId) const
 	{
+		checkSlow(AggregatedValues != nullptr);
 		return AggregatedValues->FindChecked(ColumnId);
 	}
 
@@ -177,6 +156,11 @@ public:
 
 	//////////////////////////////////////////////////
 
+	virtual uint32 GetDefaultSortOrder() const override
+	{
+		return uint32(RowId.RowIndex + 1);
+	}
+
 	virtual bool IsFiltered() const { return bIsFiltered; }
 	virtual void SetIsFiltered(bool InValue) { bIsFiltered = InValue; }
 
@@ -191,12 +175,9 @@ public:
 
 protected:
 	TWeakPtr<FTable> ParentTable;
-	FTableRowId RowId;
-
-	TMap<FName, FTableCellValue>* AggregatedValues;
-
-private:
+	FTableRowId RowId = FTableRowId::InvalidRowIndex;
 	bool bIsFiltered = false;
+	TMap<FName, FTableCellValue>* AggregatedValues = nullptr;
 };
 
 ////////////////////////////////////////////////////////////////////////////////////////////////////
@@ -207,10 +188,21 @@ class FCustomTableTreeNode : public FTableTreeNode
 
 public:
 	/** Initialization constructor for a table record node. */
-	explicit FCustomTableTreeNode(TWeakPtr<FTable> InParentTable, int32 InRowIndex, const FSlateBrush* InIconBrush, FLinearColor InColor, bool IsGroup)
+	explicit FCustomTableTreeNode(TWeakPtr<FTable> InParentTable, int32 InRowIndex, const FText& InDisplayName, bool IsGroup)
 		: FTableTreeNode(InParentTable, InRowIndex, IsGroup)
+		, DisplayName(InDisplayName)
+		, IconBrush(FBaseTreeNode::GetDefaultIcon(true))
+		, IconColor(FBaseTreeNode::GetDefaultIconColor(true))
+		, Color(FBaseTreeNode::GetDefaultColor(true))
+	{
+	}
+
+	/** Initialization constructor for a table record node. */
+	explicit FCustomTableTreeNode(TWeakPtr<FTable> InParentTable, int32 InRowIndex, const FText& InDisplayName, const FSlateBrush* InIconBrush, FLinearColor InIconColor, FLinearColor InColor, bool IsGroup)
+		: FTableTreeNode(InParentTable, InRowIndex, IsGroup)
+		, DisplayName(InDisplayName)
 		, IconBrush(InIconBrush)
-		, IconColor(InColor)
+		, IconColor(InIconColor)
 		, Color(InColor)
 	{
 	}
@@ -218,6 +210,7 @@ public:
 	/** Initialization constructor for a table record node. */
 	explicit FCustomTableTreeNode(TWeakPtr<FTable> InParentTable, int32 InRowIndex, const FSlateBrush* InIconBrush, FLinearColor InIconColor, FLinearColor InColor, bool IsGroup)
 		: FTableTreeNode(InParentTable, InRowIndex, IsGroup)
+		, DisplayName()
 		, IconBrush(InIconBrush)
 		, IconColor(InIconColor)
 		, Color(InColor)
@@ -225,19 +218,21 @@ public:
 	}
 
 	/** Initialization constructor for the group node. */
-	explicit FCustomTableTreeNode(TWeakPtr<FTable> InParentTable, FLinearColor InColor)
+	explicit FCustomTableTreeNode(TWeakPtr<FTable> InParentTable, const FText& InDisplayName)
 		: FTableTreeNode(InParentTable)
+		, DisplayName(InDisplayName)
 		, IconBrush(FBaseTreeNode::GetDefaultIcon(true))
-		, IconColor(InColor)
-		, Color(InColor)
+		, IconColor(FBaseTreeNode::GetDefaultIconColor(true))
+		, Color(FBaseTreeNode::GetDefaultColor(true))
 	{
 	}
 
 	/** Initialization constructor for the group node. */
-	explicit FCustomTableTreeNode(TWeakPtr<FTable> InParentTable, const FSlateBrush* InIconBrush, FLinearColor InColor)
+	explicit FCustomTableTreeNode(TWeakPtr<FTable> InParentTable, const FText& InDisplayName, const FSlateBrush* InIconBrush, FLinearColor InIconColor, FLinearColor InColor)
 		: FTableTreeNode(InParentTable)
+		, DisplayName(InDisplayName)
 		, IconBrush(InIconBrush)
-		, IconColor(InColor)
+		, IconColor(InIconColor)
 		, Color(InColor)
 	{
 	}
@@ -245,6 +240,7 @@ public:
 	/** Initialization constructor for the group node. */
 	explicit FCustomTableTreeNode(TWeakPtr<FTable> InParentTable, const FSlateBrush* InIconBrush, FLinearColor InIconColor, FLinearColor InColor)
 		: FTableTreeNode(InParentTable)
+		, DisplayName()
 		, IconBrush(InIconBrush)
 		, IconColor(InIconColor)
 		, Color(InColor)
@@ -255,54 +251,24 @@ public:
 	// Deprecated constructors
 
 	/** Initialization constructor for a table record node. */
-	explicit FCustomTableTreeNode(const FName InName, TWeakPtr<FTable> InParentTable, int32 InRowIndex, const FSlateBrush* InIconBrush, FLinearColor InColor, bool IsGroup)
-		: FTableTreeNode(InName, InParentTable, InRowIndex, IsGroup)
-		, DisplayName(FText::FromName(InName))
-		, IconBrush(InIconBrush)
-		, IconColor(InColor)
-		, Color(InColor)
-	{
-	}
+	//UE_DEPRECATED(5.8, "Use the constructor without FName. Set name using SetDisplayName().")
+	UE_API explicit FCustomTableTreeNode(const FName InName, TWeakPtr<FTable> InParentTable, int32 InRowIndex, const FSlateBrush* InIconBrush, FLinearColor InColor, bool IsGroup);
 
 	/** Initialization constructor for a table record node. */
-	explicit FCustomTableTreeNode(const FName InName, TWeakPtr<FTable> InParentTable, int32 InRowIndex, const FSlateBrush* InIconBrush, FLinearColor InIconColor, FLinearColor InColor, bool IsGroup)
-		: FTableTreeNode(InName, InParentTable, InRowIndex, IsGroup)
-		, DisplayName(FText::FromName(InName))
-		, IconBrush(InIconBrush)
-		, IconColor(InIconColor)
-		, Color(InColor)
-	{
-	}
+	//UE_DEPRECATED(5.8, "Use the constructor without FName. Set name using SetDisplayName().")
+	UE_API explicit FCustomTableTreeNode(const FName InName, TWeakPtr<FTable> InParentTable, int32 InRowIndex, const FSlateBrush* InIconBrush, FLinearColor InIconColor, FLinearColor InColor, bool IsGroup);
 
 	/** Initialization constructor for the group node. */
-	explicit FCustomTableTreeNode(const FName InName, TWeakPtr<FTable> InParentTable, FLinearColor InColor)
-		: FTableTreeNode(InName, InParentTable)
-		, DisplayName(FText::FromName(InName))
-		, IconBrush(FBaseTreeNode::GetDefaultIcon(true))
-		, IconColor(InColor)
-		, Color(InColor)
-	{
-	}
+	//UE_DEPRECATED(5.8, "Use the constructor without FName. Set name using SetDisplayName().")
+	UE_API explicit FCustomTableTreeNode(const FName InName, TWeakPtr<FTable> InParentTable, FLinearColor InColor);
 
 	/** Initialization constructor for the group node. */
-	explicit FCustomTableTreeNode(const FName InName, TWeakPtr<FTable> InParentTable, const FSlateBrush* InIconBrush, FLinearColor InColor)
-		: FTableTreeNode(InName, InParentTable)
-		, DisplayName(FText::FromName(InName))
-		, IconBrush(InIconBrush)
-		, IconColor(InColor)
-		, Color(InColor)
-	{
-	}
+	//UE_DEPRECATED(5.8, "Use the constructor without FName. Set name using SetDisplayName().")
+	UE_API explicit FCustomTableTreeNode(const FName InName, TWeakPtr<FTable> InParentTable, const FSlateBrush* InIconBrush, FLinearColor InColor);
 
 	/** Initialization constructor for the group node. */
-	explicit FCustomTableTreeNode(const FName InName, TWeakPtr<FTable> InParentTable, const FSlateBrush* InIconBrush, FLinearColor InIconColor, FLinearColor InColor)
-		: FTableTreeNode(InName, InParentTable)
-		, DisplayName(FText::FromName(InName))
-		, IconBrush(InIconBrush)
-		, IconColor(InIconColor)
-		, Color(InColor)
-	{
-	}
+	//UE_DEPRECATED(5.8, "Use the constructor without FName. Set name using SetDisplayName().")
+	UE_API explicit FCustomTableTreeNode(const FName InName, TWeakPtr<FTable> InParentTable, const FSlateBrush* InIconBrush, FLinearColor InIconColor, FLinearColor InColor);
 
 	//////////////////////////////////////////////////
 
