[Insights] TraceInsightsCore: Removed all usages of deprecated FName from FTableTreeNode and FCustomTableTreeNode:
* Changed FSorterByName to sort by DisplayName.
* Changed FNameValueGetter::GetValue to return DisplayName (now same as FDisplayNameValueGetter).
* Changed all instances of FCustomTableTreeNode in groupings (ByUniqueValue and PathBreakdown) to use the constructor with DisplayName.
* Fixed all instances of FTableTreeNode to avoid the deprecated constructors with FName.
* Changed "Item Name" in tooltip of tree nodes to show DisplayName.

#rnx
#Synced-CL 51291823

diff --git a/Engine/Source/Developer/TraceInsightsCore/Private/InsightsCore/Table/ViewModels/TableCellValueGetter.cpp b/Engine/Source/Developer/TraceInsightsCore/Private/InsightsCore/Table/ViewModels/TableCellValueGetter.cpp
--- a/Engine/Source/Developer/TraceInsightsCore/Private/InsightsCore/Table/ViewModels/TableCellValueGetter.cpp
+++ b/Engine/Source/Developer/TraceInsightsCore/Private/InsightsCore/Table/ViewModels/TableCellValueGetter.cpp
@@ -13,7 +13,7 @@ namespace UE::Insights
 
 const TOptional<FTableCellValue> FNameValueGetter::GetValue(const FTableColumn& Column, const FBaseTreeNode& Node) const
 {
-	return TOptional<FTableCellValue>(FTableCellValue(FText::FromName(Node.GetName())));
+	return TOptional<FTableCellValue>(FTableCellValue(Node.GetDisplayName()));
 }
 
 ////////////////////////////////////////////////////////////////////////////////////////////////////
diff --git a/Engine/Source/Developer/TraceInsightsCore/Private/InsightsCore/Table/ViewModels/TableCellValueSorter.cpp b/Engine/Source/Developer/TraceInsightsCore/Private/InsightsCore/Table/ViewModels/TableCellValueSorter.cpp
--- a/Engine/Source/Developer/TraceInsightsCore/Private/InsightsCore/Table/ViewModels/TableCellValueSorter.cpp
+++ b/Engine/Source/Developer/TraceInsightsCore/Private/InsightsCore/Table/ViewModels/TableCellValueSorter.cpp
@@ -111,7 +111,7 @@ FSorterByName::FSorterByName(TSharedRef<FTableColumn> InColumnRef)
 		INSIGHTS_DEFAULT_PRESORTING_NODES(A, B)
 
 		// Sort by name (ascending).
-		return A->GetName().LexicalLess(B->GetName());
+		return A->GetDisplayName().CompareTo(B->GetDisplayName()) < 0;
 	};
 
 	DescendingCompareDelegate = [this](const FBaseTreeNodePtr& A, const FBaseTreeNodePtr& B) -> bool
@@ -119,7 +119,7 @@ FSorterByName::FSorterByName(TSharedRef<FTableColumn> InColumnRef)
 		INSIGHTS_DEFAULT_PRESORTING_NODES(A, B)
 
 		// Sort by name (descending).
-		return B->GetName().LexicalLess(A->GetName());
+		return B->GetDisplayName().CompareTo(A->GetDisplayName()) < 0;
 	};
 }
 
diff --git a/Engine/Source/Developer/TraceInsightsCore/Private/InsightsCore/Table/ViewModels/TreeNodeGrouping.cpp b/Engine/Source/Developer/TraceInsightsCore/Private/InsightsCore/Table/ViewModels/TreeNodeGrouping.cpp
--- a/Engine/Source/Developer/TraceInsightsCore/Private/InsightsCore/Table/ViewModels/TreeNodeGrouping.cpp
+++ b/Engine/Source/Developer/TraceInsightsCore/Private/InsightsCore/Table/ViewModels/TreeNodeGrouping.cpp
@@ -68,11 +68,11 @@ void FTreeNodeGrouping::GroupNodes(const TArray<FTableTreeNodePtr>& Nodes, FTabl
 			if (GroupInfo.IsCustom())
 			{
 				const FSlateBrush* GroupIcon = GroupInfo.GetIcon() != nullptr ? GroupInfo.GetIcon() : FBaseTreeNode::GetDefaultIcon(true);
-				GroupPtr = MakeShared<FCustomTableTreeNode>(GroupInfo.GetName(), InParentTable, GroupIcon, GroupInfo.GetIconColor(), GroupInfo.GetColor());
+				GroupPtr = MakeShared<FCustomTableTreeNode>(InParentTable, GroupInfo.GetDisplayName(), GroupIcon, GroupInfo.GetIconColor(), GroupInfo.GetColor());
 			}
 			else
 			{
-				GroupPtr = MakeShared<FTableTreeNode>(GroupInfo.GetName(), InParentTable);
+				GroupPtr = MakeShared<FCustomTableTreeNode>(InParentTable, GroupInfo.GetDisplayName());
 			}
 			GroupPtr->SetExpansion(GroupInfo.IsExpanded());
 			ParentGroup.AddChildAndSetParent(GroupPtr);
@@ -114,19 +114,21 @@ void FTreeNodeGroupingFlat::GroupNodes(const TArray<FTableTreeNodePtr>& Nodes, F
 {
 	ParentGroup.ClearChildren(1);
 
-	FTableTreeNodePtr GroupPtr = MakeShared<FCustomTableTreeNode>(FName(TEXT("All")), InParentTable, GetColor());
-	GroupPtr->SetExpansion(true);
-	ParentGroup.AddChildAndSetParent(GroupPtr);
+	auto GroupNode = MakeShared<FCustomTableTreeNode>(InParentTable, LOCTEXT("GroupName_All", "All"));
+	GroupNode->SetIconColor(GetColor());
+	GroupNode->SetColor(GetColor());
+	GroupNode->SetExpansion(true);
+	ParentGroup.AddChildAndSetParent(GroupNode);
 
-	GroupPtr->ClearChildren(Nodes.Num());
-	for (FTableTreeNodePtr NodePtr : Nodes)
+	GroupNode->ClearChildren(Nodes.Num());
+	for (FTableTreeNodePtr ChildNode : Nodes)
 	{
 		if (InAsyncOperationProgress.ShouldCancelAsyncOp())
 		{
 			break;
 		}
 
-		GroupPtr->AddChildAndSetParent(NodePtr);
+		GroupNode->AddChildAndSetParent(ChildNode);
 	}
 }
 
@@ -151,14 +153,40 @@ FTreeNodeGroupInfo FTreeNodeGroupingByUniqueValue::GetGroupForNode(const FBaseTr
 {
 	FTableTreeNodePtr TableTreeNodePtr = StaticCastSharedPtr<FTableTreeNode>(InNode);
 	FText ValueAsText = ColumnRef->GetValueAsText(*TableTreeNodePtr);
+
 	FStringView GroupName(ValueAsText.ToString());
 	if (GroupName.Len() >= NAME_SIZE)
 	{
 		GroupName = FStringView(GroupName.GetData(), NAME_SIZE - 1);
 	}
+
 	return FTreeNodeGroupInfo(FName(GroupName, 0), false, GetColor());
 }
 
+////////////////////////////////////////////////////////////////////////////////////////////////////
+
+TSharedRef<FCustomTableTreeNode> FTreeNodeGroupingByUniqueValue::CreateGroupNode(TWeakPtr<FTable> InParentTable, const FText& InDisplayName) const
+{
+	auto GroupNode = MakeShared<FCustomTableTreeNode>(InParentTable, InDisplayName);
+	if (GetIcon() != nullptr)
+	{
+		GroupNode->SetIcon(GetIcon());
+	}
+	GroupNode->SetIconColor(GetColor());
+	GroupNode->SetColor(GetColor());
+	GroupNode->SetExpansion(false);
+	return GroupNode;
+}
+
+////////////////////////////////////////////////////////////////////////////////////////////////////
+
+TSharedRef<FCustomTableTreeNode> FTreeNodeGroupingByUniqueValue::CreateUnsetGroupNode(TWeakPtr<FTable> InParentTable) const
+{
+	auto GroupNode = MakeShared<FCustomTableTreeNode>(InParentTable, LOCTEXT("GroupName_Unset", "<unset>"));
+	GroupNode->SetExpansion(false);
+	return GroupNode;
+}
+
 ////////////////////////////////////////////////////////////////////////////////////////////////////
 // TTreeNodeGroupingByUniqueValue specializations
 ////////////////////////////////////////////////////////////////////////////////////////////////////
@@ -219,11 +247,15 @@ void FTreeNodeGroupingByUniqueValueCString::GroupNodes(const TArray<FTableTreeNo
 				if (!GroupPtrPtr)
 				{
 					const FName GroupName = GetGroupName(Value);
-					GroupPtrPtr = GroupNameMap.Find(Value);
+					GroupPtrPtr = GroupNameMap.Find(GroupName);
 					if (!GroupPtrPtr)
 					{
-						GroupPtr = MakeShared<FCustomTableTreeNode>(GroupName, InParentTable, GetColor());
-						GroupPtr->SetExpansion(false);
+						auto GroupNode = MakeShared<FCustomTableTreeNode>(InParentTable, FText::FromString(Value));
+						GroupNode->SetIconColor(GetColor());
+						GroupNode->SetColor(GetColor());
+						GroupNode->SetExpansion(false);
+
+						GroupPtr = GroupNode;
 						ParentGroup.AddChildAndSetParent(GroupPtr);
 						GroupNameMap.Add(GroupName, GroupPtr);
 					}
@@ -242,7 +274,7 @@ void FTreeNodeGroupingByUniqueValueCString::GroupNodes(const TArray<FTableTreeNo
 			{
 				if (!EmptyGroupPtr)
 				{
-					EmptyGroupPtr = MakeShared<FTableTreeNode>(FName(TEXT("N/A")), InParentTable);
+					EmptyGroupPtr = MakeShared<FCustomTableTreeNode>(InParentTable, LOCTEXT("GroupName_NA", "N/A"));
 					EmptyGroupPtr->SetExpansion(false);
 					ParentGroup.AddChildAndSetParent(EmptyGroupPtr);
 				}
@@ -253,7 +285,7 @@ void FTreeNodeGroupingByUniqueValueCString::GroupNodes(const TArray<FTableTreeNo
 		{
 			if (!UnsetGroupPtr)
 			{
-				UnsetGroupPtr = MakeShared<FTableTreeNode>(FName(TEXT("<unset>")), InParentTable);
+				UnsetGroupPtr = MakeShared<FCustomTableTreeNode>(InParentTable, LOCTEXT("GroupName_Unset", "<unset>"));
 				UnsetGroupPtr->SetExpansion(false);
 				ParentGroup.AddChildAndSetParent(UnsetGroupPtr);
 			}
@@ -281,7 +313,7 @@ FTreeNodeGroupingByNameFirstLetter::FTreeNodeGroupingByNameFirstLetter()
 
 FTreeNodeGroupInfo FTreeNodeGroupingByNameFirstLetter::GetGroupForNode(const FBaseTreeNodePtr InNode) const
 {
-	return { *InNode->GetName().GetPlainNameString().Left(1), false };
+	return { *InNode->GetDisplayName().ToString().Left(1), false };
 }
 
 ////////////////////////////////////////////////////////////////////////////////////////////////////
@@ -400,10 +432,17 @@ void FTreeNodeGroupingByPathBreakdown::GroupNodes(
 						}
 						else if (!DirName.IsNone())
 						{
-							FTableTreeNodePtr GroupPtr = MakeShared<FCustomTableTreeNode>(DirName, InParentTable, GetColor());
-							GroupPtr->SetExpansion(false);
-							Directory->Node->AddChildAndSetParent(GroupPtr);
-							Directory = new FDirectory(DirName, Directory, GroupPtr.Get());
+							const FText GroupName = FText::FromName(DirName);
+							auto GroupNode = MakeShared<FCustomTableTreeNode>(InParentTable, GroupName);
+							if (GetIcon() != nullptr)
+							{
+								GroupNode->SetIcon(GetIcon());
+							}
+							GroupNode->SetIconColor(GetColor());
+							GroupNode->SetColor(GetColor());
+							GroupNode->SetExpansion(false);
+							Directory->Node->AddChildAndSetParent(GroupNode);
+							Directory = new FDirectory(DirName, Directory, &GroupNode.Get());
 							Directory->Parent->ChildrenMap.Add(DirName, Directory);
 							Directories.Add(Directory);
 						}
@@ -429,11 +468,10 @@ void FTreeNodeGroupingByPathBreakdown::GroupNodes(
 		{
 			if (!UnsetGroupPtr)
 			{
-				static FName NotAvailableName(TEXT("N/A"));
-				FTableTreeNodePtr NotAvailableNode = MakeShared<FTableTreeNode>(NotAvailableName, InParentTable);
-				NotAvailableNode->SetExpansion(false);
-				ParentGroup.AddChildAndSetParent(NotAvailableNode);
-				UnsetGroupPtr = NotAvailableNode.Get();
+				auto GroupNode = MakeShared<FCustomTableTreeNode>(InParentTable, LOCTEXT("GroupName_NA", "N/A"));
+				GroupNode->SetExpansion(false);
+				ParentGroup.AddChildAndSetParent(GroupNode);
+				UnsetGroupPtr = &GroupNode.Get();
 			}
 			UnsetGroupPtr->AddChildAndSetParent(NodePtr);
 		}
diff --git a/Engine/Source/Developer/TraceInsightsCore/Private/InsightsCore/Table/Widgets/STableTreeView.cpp b/Engine/Source/Developer/TraceInsightsCore/Private/InsightsCore/Table/Widgets/STableTreeView.cpp
--- a/Engine/Source/Developer/TraceInsightsCore/Private/InsightsCore/Table/Widgets/STableTreeView.cpp
+++ b/Engine/Source/Developer/TraceInsightsCore/Private/InsightsCore/Table/Widgets/STableTreeView.cpp
@@ -130,16 +130,13 @@ public:
 // STableTreeView
 ////////////////////////////////////////////////////////////////////////////////////////////////////
 
-const FName STableTreeView::RootNodeName(TEXT("Root"));
-
-////////////////////////////////////////////////////////////////////////////////////////////////////
-
 STableTreeView::STableTreeView()
-	: Root(MakeShared<FTableTreeNode>(RootNodeName, Table))
+	: Root(MakeShared<FCustomTableTreeNode>(Table, LOCTEXT("GroupName_Root", "Root")))
 	, FilteredNodesPtr(&TableRowNodes)
 	, ColumnBeingSorted(GetDefaultColumnBeingSorted())
 	, ColumnSortMode(GetDefaultColumnSortMode())
 {
+	check(Root->IsGroup());
 }
 
 ////////////////////////////////////////////////////////////////////////////////////////////////////
diff --git a/Engine/Source/Developer/TraceInsightsCore/Private/InsightsCore/Table/Widgets/STableTreeViewTooltip.cpp b/Engine/Source/Developer/TraceInsightsCore/Private/InsightsCore/Table/Widgets/STableTreeViewTooltip.cpp
--- a/Engine/Source/Developer/TraceInsightsCore/Private/InsightsCore/Table/Widgets/STableTreeViewTooltip.cpp
+++ b/Engine/Source/Developer/TraceInsightsCore/Private/InsightsCore/Table/Widgets/STableTreeViewTooltip.cpp
@@ -164,12 +164,12 @@ TSharedPtr<SToolTip> STableTreeViewTooltip::GetRowTooltip(const TSharedPtr<FTabl
 						.TextStyle(GetTableTreeViewTooltipNormalTextStyle())
 					]
 
-					// Item Name: [Name]
+					// Name: [Name]
 					+ SGridPanel::Slot(0, 2)
 					.Padding(2.0f)
 					[
 						SNew(STextBlock)
-						.Text(LOCTEXT("TT_Name", "Item Name:"))
+						.Text(LOCTEXT("TT_Name", "Name:"))
 						.TextStyle(GetTableTreeViewTooltipHeaderTextStyle())
 					]
 					+ SGridPanel::Slot(1, 2)
@@ -178,7 +178,7 @@ TSharedPtr<SToolTip> STableTreeViewTooltip::GetRowTooltip(const TSharedPtr<FTabl
 						SNew(STextBlock)
 						.WrapTextAt(512.0f)
 						.WrappingPolicy(ETextWrappingPolicy::AllowPerCharacterWrapping)
-						.Text(FText::FromName(TreeNodePtr->GetName()))
+						.Text(TreeNodePtr->GetDisplayName())
 						.TextStyle(GetTableTreeViewTooltipNormalTextStyle())
 					]
 				]
diff --git a/Engine/Source/Developer/TraceInsightsCore/Private/InsightsCore/Table/Widgets/SUntypedTableTreeView.cpp b/Engine/Source/Developer/TraceInsightsCore/Private/InsightsCore/Table/Widgets/SUntypedTableTreeView.cpp
--- a/Engine/Source/Developer/TraceInsightsCore/Private/InsightsCore/Table/Widgets/SUntypedTableTreeView.cpp
+++ b/Engine/Source/Developer/TraceInsightsCore/Private/InsightsCore/Table/Widgets/SUntypedTableTreeView.cpp
@@ -86,13 +86,10 @@ void SUntypedTableTreeView::RebuildTree(bool bResync)
 		if (TotalRowCount != TableRowNodes.Num())
 		{
 			TableRowNodes.Empty(TotalRowCount);
-			FName BaseNodeName(TEXT("row"));
 			for (int32 RowIndex = 0; RowIndex < TotalRowCount; ++RowIndex)
 			{
 				TableReader->SetRowIndex(RowIndex);
-				FName NodeName(BaseNodeName, RowIndex + 1);
-				FTableTreeNodePtr NodePtr = MakeShared<FTableTreeNode>(NodeName, Table, RowIndex);
-				NodePtr->SetDefaultSortOrder(RowIndex + 1);
+				FTableTreeNodePtr NodePtr = MakeShared<FTableTreeNode>(Table, RowIndex);
 				TableRowNodes.Add(NodePtr);
 			}
 			ensure(TableRowNodes.Num() == TotalRowCount);
diff --git a/Engine/Source/Developer/TraceInsightsCore/Public/InsightsCore/Table/ViewModels/TreeNodeGrouping.h b/Engine/Source/Developer/TraceInsightsCore/Public/InsightsCore/Table/ViewModels/TreeNodeGrouping.h
--- a/Engine/Source/Developer/TraceInsightsCore/Public/InsightsCore/Table/ViewModels/TreeNodeGrouping.h
+++ b/Engine/Source/Developer/TraceInsightsCore/Public/InsightsCore/Table/ViewModels/TreeNodeGrouping.h
@@ -80,6 +80,7 @@ public:
 	}
 
 	const FName& GetName() const { return Name; }
+	FText GetDisplayName() const { return FText::FromName(Name); }
 
 	bool IsExpanded() const { return bIsExpanded; }
 	bool IsCustom() const { return bIsCustom; }
@@ -166,6 +167,10 @@ public:
 	virtual FName GetColumnId() const override { return ColumnRef->GetId(); }
 	TSharedRef<FTableColumn> GetColumn() const { return ColumnRef; }
 
+protected:
+	TSharedRef<FCustomTableTreeNode> CreateGroupNode(TWeakPtr<FTable> InParentTable, const FText& InGroupName) const;
+	TSharedRef<FCustomTableTreeNode> CreateUnsetGroupNode(TWeakPtr<FTable> InParentTable) const;
+
 private:
 	TSharedRef<FTableColumn> ColumnRef;
 };
@@ -184,28 +189,8 @@ public:
 
 private:
 	static Type GetValue(const FTableCellValue& CellValue);
-	static FName GetGroupName(const FTableColumn& Column, const FTableTreeNode& Node);
 };
 
-template<typename Type>
-FName TTreeNodeGroupingByUniqueValue<Type>::GetGroupName(const FTableColumn& Column, const FTableTreeNode& Node)
-{
-	FText ValueAsText = Column.GetValueAsGroupingText(Node);
-
-	if (ValueAsText.IsEmpty())
-	{
-		static FName EmptyGroupName(TEXT("N/A"));
-		return EmptyGroupName;
-	}
-
-	FStringView StringView(ValueAsText.ToString());
-	if (StringView.Len() >= NAME_SIZE)
-	{
-		StringView = FStringView(StringView.GetData(), NAME_SIZE - 1);
-	}
-	return FName(StringView, 0);
-}
-
 template<typename Type>
 void TTreeNodeGroupingByUniqueValue<Type>::GroupNodes(const TArray<FTableTreeNodePtr>& Nodes, FTableTreeNode& ParentGroup, TWeakPtr<FTable> InParentTable, IAsyncOperationProgress& InAsyncOperationProgress) const
 {
@@ -238,9 +223,8 @@ void TTreeNodeGroupingByUniqueValue<Type>::GroupNodes(const TArray<FTableTreeNod
 			FTableTreeNodePtr* GroupPtrPtr = GroupMap.Find(Value);
 			if (!GroupPtrPtr)
 			{
-				const FName GroupName = GetGroupName(Column, *NodePtr);
-				GroupPtr = MakeShared<FCustomTableTreeNode>(GroupName, InParentTable, GetColor());
-				GroupPtr->SetExpansion(false);
+				const FText GroupName = Column.GetValueAsGroupingText(*NodePtr);
+				GroupPtr = CreateGroupNode(InParentTable, GroupName);
 				ParentGroup.AddChildAndSetParent(GroupPtr);
 				GroupMap.Add(Value, GroupPtr);
 			}
@@ -253,8 +237,7 @@ void TTreeNodeGroupingByUniqueValue<Type>::GroupNodes(const TArray<FTableTreeNod
 		{
 			if (!UnsetGroupPtr)
 			{
-				UnsetGroupPtr = MakeShared<FTableTreeNode>(FName(TEXT("<unset>")), InParentTable);
-				UnsetGroupPtr->SetExpansion(false);
+				UnsetGroupPtr = CreateUnsetGroupNode(InParentTable);
 				ParentGroup.AddChildAndSetParent(UnsetGroupPtr);
 			}
 			GroupPtr = UnsetGroupPtr;
diff --git a/Engine/Source/Developer/TraceInsightsCore/Public/InsightsCore/Table/Widgets/STableTreeView.h b/Engine/Source/Developer/TraceInsightsCore/Public/InsightsCore/Table/Widgets/STableTreeView.h
--- a/Engine/Source/Developer/TraceInsightsCore/Public/InsightsCore/Table/Widgets/STableTreeView.h
+++ b/Engine/Source/Developer/TraceInsightsCore/Public/InsightsCore/Table/Widgets/STableTreeView.h
@@ -539,8 +539,6 @@ protected:
 	//////////////////////////////////////////////////
 	// Tree Nodes
 
-	static UE_API const FName RootNodeName;
-
 	/** The root node of the tree. It is invisible; only its children are displayed in the tree view. */
 	FTableTreeNodePtr Root;
 
