[Insights] ObjectProfiler: Added initial implementation for SSegmentedBarGraph widget.

#Synced-CL 52097700

diff --git a/Engine/Source/Developer/TraceInsightsCore/Private/InsightsCore/Widgets/SSegmentedBarGraph.cpp b/Engine/Source/Developer/TraceInsightsCore/Private/InsightsCore/Widgets/SSegmentedBarGraph.cpp
new file mode 100644
--- /dev/null
+++ b/Engine/Source/Developer/TraceInsightsCore/Private/InsightsCore/Widgets/SSegmentedBarGraph.cpp
@@ -0,0 +1,173 @@
+// Copyright Epic Games, Inc. All Rights Reserved.
+
+#include "InsightsCore/Widgets/SSegmentedBarGraph.h"
+
+#include "SlateOptMacros.h"
+#include "Styling/AppStyle.h"
+#include "Widgets/Images/SImage.h"
+#include "Widgets/Layout/SBorder.h"
+#include "Widgets/Text/STextBlock.h"
+
+#define LOCTEXT_NAMESPACE "UE::Insights::SSegmentedBarGraph"
+
+namespace UE::Insights
+{
+
+////////////////////////////////////////////////////////////////////////////////////////////////////
+
+void FSharedObservableSegmentArray::HandleArrayChanged(typename ::UE::Slate::Containers::TObservableArray<IBarGraphSegment>::ObservableArrayChangedArgsType Args)
+{
+	if (TSharedPtr<SSegmentedBarGraph> GraphOwnerPin = GraphOwner.Pin())
+	{
+		GraphOwnerPin->RequestGraphRefresh();
+	}
+}
+
+////////////////////////////////////////////////////////////////////////////////////////////////////
+// SSegmentedBarGraph
+////////////////////////////////////////////////////////////////////////////////////////////////////
+
+SSegmentedBarGraph::SSegmentedBarGraph()
+{
+}
+
+////////////////////////////////////////////////////////////////////////////////////////////////////
+
+SSegmentedBarGraph::~SSegmentedBarGraph()
+{
+}
+
+////////////////////////////////////////////////////////////////////////////////////////////////////
+
+BEGIN_SLATE_FUNCTION_BUILD_OPTIMIZATION
+void SSegmentedBarGraph::Construct(const FArguments& InArgs)
+{
+	SAssignNew(HorizontalBox, SHorizontalBox);
+	check(HorizontalBox.IsValid());
+	ChildSlot[ HorizontalBox.ToSharedRef() ];
+
+	RequestGraphRefresh();
+}
+END_SLATE_FUNCTION_BUILD_OPTIMIZATION
+
+////////////////////////////////////////////////////////////////////////////////////////////////////
+
+void SSegmentedBarGraph::RequestGraphRefresh()
+{
+	check(HorizontalBox.IsValid())
+	HorizontalBox->ClearChildren();
+
+	TArrayView<const IBarGraphSegment> Segments = GetSegments();
+	for (int32 SegmentIndex = 0; SegmentIndex < Segments.Num(); ++SegmentIndex)
+	{
+		AddHorizontalSegment(HorizontalBox.ToSharedRef(), SegmentIndex);
+	}
+}
+
+////////////////////////////////////////////////////////////////////////////////////////////////////
+
+BEGIN_SLATE_FUNCTION_BUILD_OPTIMIZATION
+void SSegmentedBarGraph::AddHorizontalSegment(TSharedRef<SHorizontalBox> Box, int32 SegmentIndex)
+{
+	Box->AddSlot()
+		.HAlign(HAlign_Fill)
+		.VAlign(VAlign_Fill)
+		.Padding(FMargin(0.0f, 0.0f))
+		.FillWidth(TAttribute<float>::CreateSP(this, &SSegmentedBarGraph::GetSegmentSize, SegmentIndex))
+		[
+			SNew(SBorder)
+			.BorderImage(FAppStyle::Get().GetBrush("WhiteBrush"))
+			.BorderBackgroundColor(TAttribute<FSlateColor>::CreateSP(this, &SSegmentedBarGraph::GetSegmentColor, SegmentIndex))
+			.ColorAndOpacity(TAttribute<FLinearColor>::CreateSP(this, &SSegmentedBarGraph::GetSegmentBorderColor, SegmentIndex))
+			.ToolTipText(this, &SSegmentedBarGraph::GetSegmentToolTipText, SegmentIndex)
+			[
+				SNew(STextBlock)
+				.ColorAndOpacity(TAttribute<FSlateColor>::CreateSP(this, &SSegmentedBarGraph::GetSegmentTextColor, SegmentIndex))
+				.Text(this, &SSegmentedBarGraph::GetSegmentText, SegmentIndex)
+			]
+		];
+}
+END_SLATE_FUNCTION_BUILD_OPTIMIZATION
+
+////////////////////////////////////////////////////////////////////////////////////////////////////
+
+float SSegmentedBarGraph::GetSegmentSize(int32 SegmentIndex) const
+{
+	TArrayView<const IBarGraphSegment> Segments = GetSegments();
+	if (SegmentIndex < 0 || SegmentIndex >= Segments.Num())
+	{
+		return 0;
+	}
+	return static_cast<float>(Segments[SegmentIndex].GetSize());
+}
+
+////////////////////////////////////////////////////////////////////////////////////////////////////
+
+FText SSegmentedBarGraph::GetSegmentText(int32 SegmentIndex) const
+{
+	TArrayView<const IBarGraphSegment> Segments = GetSegments();
+	if (SegmentIndex < 0 || SegmentIndex >= Segments.Num())
+	{
+		return FText::GetEmpty();
+	}
+	return Segments[SegmentIndex].GetText();
+}
+
+////////////////////////////////////////////////////////////////////////////////////////////////////
+
+FText SSegmentedBarGraph::GetSegmentToolTipText(int32 SegmentIndex) const
+{
+	TArrayView<const IBarGraphSegment> Segments = GetSegments();
+	if (SegmentIndex < 0 || SegmentIndex >= Segments.Num())
+	{
+		return FText::GetEmpty();
+	}
+	return Segments[SegmentIndex].GetToolTipText();
+}
+
+////////////////////////////////////////////////////////////////////////////////////////////////////
+
+FSlateColor SSegmentedBarGraph::GetSegmentColor(int32 SegmentIndex) const
+{
+	TArrayView<const IBarGraphSegment> Segments = GetSegments();
+	if (SegmentIndex < 0 || SegmentIndex >= Segments.Num())
+	{
+		FLinearColor Color(0.75f, 0.75f, 0.75f, 1.0f);
+		return FSlateColor(Color);
+	}
+	return FSlateColor(Segments[SegmentIndex].GetColor());
+}
+
+////////////////////////////////////////////////////////////////////////////////////////////////////
+
+FSlateColor SSegmentedBarGraph::GetSegmentTextColor(int32 SegmentIndex) const
+{
+	TArrayView<const IBarGraphSegment> Segments = GetSegments();
+	if (SegmentIndex < 0 || SegmentIndex >= Segments.Num())
+	{
+		FLinearColor Color(0.0f, 0.0f, 0.0f, 1.0f);
+		return FSlateColor(Color);
+	}
+	return FSlateColor(Segments[SegmentIndex].GetTextColor());
+}
+
+////////////////////////////////////////////////////////////////////////////////////////////////////
+
+FLinearColor SSegmentedBarGraph::GetSegmentBorderColor(int32 SegmentIndex) const
+{
+	TArrayView<const IBarGraphSegment> Segments = GetSegments();
+	if (SegmentIndex < 0 || SegmentIndex >= Segments.Num())
+	{
+		return FLinearColor(0.56f, 0.56f, 0.56f, 1.0f);
+	}
+	FLinearColor Color = Segments[SegmentIndex].GetColor();
+	Color *= 0.75f;
+	Color.A = 1.0f;
+	return Color;
+}
+
+////////////////////////////////////////////////////////////////////////////////////////////////////
+
+} // namespace UE::Insights
+
+#undef LOCTEXT_NAMESPACE
diff --git a/Engine/Source/Developer/TraceInsightsCore/Public/InsightsCore/ViewModels/IBarGraphSegment.h b/Engine/Source/Developer/TraceInsightsCore/Public/InsightsCore/ViewModels/IBarGraphSegment.h
new file mode 100644
--- /dev/null
+++ b/Engine/Source/Developer/TraceInsightsCore/Public/InsightsCore/ViewModels/IBarGraphSegment.h
@@ -0,0 +1,45 @@
+// Copyright Epic Games, Inc. All Rights Reserved.
+
+#pragma once
+
+#include "CoreTypes.h"
+
+#include "Internationalization/Text.h"
+#include "Math/Color.h"
+#include "Templates/SharedPointer.h"
+
+namespace UE::Insights
+{
+
+class IBarGraphSegment : public TSharedFromThis<IBarGraphSegment>
+{
+public:
+	virtual ~IBarGraphSegment() = default;
+
+	/**
+	 * Returns the relative segment size, in interval [0 .. 1].
+	 */
+	virtual double GetSize() const = 0;
+
+	/**
+	 * Returns the text displayed on the segment.
+	 */
+	virtual FText GetText() const = 0;
+
+	/**
+	 * Returns the tool tip text of the segment.
+	 */
+	virtual FText GetToolTipText() const = 0;
+
+	/**
+	 * Returns the color of the segment.
+	 */
+	virtual FLinearColor GetColor() const = 0;
+
+	/**
+	 * Returns the color of the text displayed on the segment.
+	 */
+	virtual FLinearColor GetTextColor() const = 0;
+};
+
+} // namespace UE::Insights
diff --git a/Engine/Source/Developer/TraceInsightsCore/Public/InsightsCore/Widgets/SSegmentedBarGraph.h b/Engine/Source/Developer/TraceInsightsCore/Public/InsightsCore/Widgets/SSegmentedBarGraph.h
new file mode 100644
--- /dev/null
+++ b/Engine/Source/Developer/TraceInsightsCore/Public/InsightsCore/Widgets/SSegmentedBarGraph.h
@@ -0,0 +1,154 @@
+// Copyright Epic Games, Inc. All Rights Reserved.
+
+#pragma once
+
+#include "CoreTypes.h"
+
+#include "Containers/Array.h"
+#include "Containers/ArrayView.h"
+#include "Widgets/DeclarativeSyntaxSupport.h"
+#include "Widgets/SCompoundWidget.h"
+#include "Widgets/Views/IItemsSource.h"
+
+// InsightsCore
+#include "InsightsCore/ViewModels/IBarGraphSegment.h"
+
+#define UE_API TRACEINSIGHTSCORE_API
+
+namespace UE::Insights
+{
+
+////////////////////////////////////////////////////////////////////////////////////////////////////
+
+class SSegmentedBarGraph;
+
+class FSharedObservableSegmentArray : public UE::Slate::ItemsSource::IItemsSource<IBarGraphSegment>
+{
+public:
+	explicit FSharedObservableSegmentArray(TSharedRef<SSegmentedBarGraph> InGraph, TSharedRef<::UE::Slate::Containers::TObservableArray<IBarGraphSegment>> InItemsSource)
+		: ItemsSource(InItemsSource)
+		, GraphOwner(InGraph.ToWeakPtr())
+	{
+		ArrayChangedHandle = InItemsSource->OnArrayChanged().AddRaw(this, &FSharedObservableSegmentArray::HandleArrayChanged);
+	}
+
+	virtual ~FSharedObservableSegmentArray()
+	{
+		ItemsSource->OnArrayChanged().Remove(ArrayChangedHandle);
+	}
+
+	virtual const TArrayView<const IBarGraphSegment> GetItems() const override
+	{
+		return TArrayView<const IBarGraphSegment>(ItemsSource->GetData(), ItemsSource->Num());
+	}
+
+	virtual bool IsSame(const void* RawPointer) const override
+	{
+		UE::Slate::Containers::TObservableArray<IBarGraphSegment>* ValueToTest = &(ItemsSource.Get());
+		return RawPointer == reinterpret_cast<const void*>(ValueToTest);
+	}
+
+private:
+	void HandleArrayChanged(typename ::UE::Slate::Containers::TObservableArray<IBarGraphSegment>::ObservableArrayChangedArgsType Args);
+
+private:
+	TSharedRef<UE::Slate::Containers::TObservableArray<IBarGraphSegment>> ItemsSource;
+	TWeakPtr<SSegmentedBarGraph> GraphOwner;
+	FDelegateHandle ArrayChangedHandle;
+};
+
+////////////////////////////////////////////////////////////////////////////////////////////////////
+
+/**
+ * A widget that displays a segmented bar graph.
+ */
+class SSegmentedBarGraph : public SCompoundWidget
+{
+public:
+	/** Default constructor. */
+	UE_API SSegmentedBarGraph();
+
+	/** Virtual destructor. */
+	UE_API virtual ~SSegmentedBarGraph();
+
+	SLATE_BEGIN_ARGS(SSegmentedBarGraph) {}
+		SLATE_ITEMS_SOURCE_ARGUMENT(IBarGraphSegment, SegmentsSource)
+	SLATE_END_ARGS()
+
+	/**
+	 * Construct this widget.
+	 * @param InArgs - The declaration data for this widget
+	 */
+	UE_API void Construct(const FArguments& InArgs);
+
+	UE_API void RequestGraphRefresh();
+
+	/**
+	 * Set the SegmentsSource. The graph bar will generate widgets to represent these segments.
+	 * @param InSegmentsSource A pointer to the array of segments that should be observed by this widget.
+	*/
+	void SetSegmentsSource(const TArray<IBarGraphSegment>* InSegmentsSource)
+	{
+		ensureMsgf(InSegmentsSource, TEXT("The SegmentsSource is invalid."));
+		if (SegmentsSource == nullptr || !SegmentsSource->IsSame(reinterpret_cast<const void*>(InSegmentsSource)))
+		{
+			if (InSegmentsSource)
+			{
+				SetSegmentsSource(MakeUnique<UE::Slate::ItemsSource::FArrayPointer<IBarGraphSegment>>(InSegmentsSource));
+			}
+			else
+			{
+				ClearSegmentsSource();
+			}
+		}
+	}
+
+	void SetSegmentsSource(TSharedRef<UE::Slate::Containers::TObservableArray<IBarGraphSegment>> InSegmentsSource)
+	{
+		if (SegmentsSource == nullptr || !SegmentsSource->IsSame(reinterpret_cast<const void*>(&InSegmentsSource.Get())))
+		{
+			SetSegmentsSource(MakeUnique<FSharedObservableSegmentArray>(this->SharedThis(this), MoveTemp(InSegmentsSource)));
+		}
+	}
+
+	void SetSegmentsSource(TUniquePtr<UE::Slate::ItemsSource::IItemsSource<IBarGraphSegment>> Provider)
+	{
+		SegmentsSource = MoveTemp(Provider);
+		RequestGraphRefresh();
+	}
+
+	void ClearSegmentsSource()
+	{
+		SetSegmentsSource(TUniquePtr<UE::Slate::ItemsSource::IItemsSource<IBarGraphSegment>>());
+	}
+
+	bool HasValidSegmentsSource() const
+	{
+		return SegmentsSource != nullptr;
+	}
+
+	TArrayView<const IBarGraphSegment> GetSegments() const
+	{
+		return SegmentsSource ? SegmentsSource->GetItems() : TArrayView<const IBarGraphSegment>();
+	}
+
+private:
+	void AddHorizontalSegment(TSharedRef<SHorizontalBox> Box, int32 SegmentIndex);
+	float GetSegmentSize(int32 SegmentIndex) const;
+	FText GetSegmentText(int32 SegmentIndex) const;
+	FText GetSegmentToolTipText(int32 SegmentIndex) const;
+	FSlateColor GetSegmentColor(int32 SegmentIndex) const;
+	FSlateColor GetSegmentTextColor(int32 SegmentIndex) const;
+	FLinearColor GetSegmentBorderColor(int32 SegmentIndex) const;
+
+private:
+	/** Pointer to the source data that we are observing */
+	TUniquePtr<UE::Slate::ItemsSource::IItemsSource<IBarGraphSegment>> SegmentsSource;
+	TSharedPtr<SHorizontalBox> HorizontalBox;
+};
+
+////////////////////////////////////////////////////////////////////////////////////////////////////
+
+} // namespace UE::Insights
+
+#undef UE_API
