[Insights] Updated implementation for SSegmentedBarGraph.

#rnx
#Synced-CL 52097700

diff --git a/Engine/Source/Developer/TraceInsightsCore/Private/InsightsCore/Widgets/SSegmentedBarGraph.cpp b/Engine/Source/Developer/TraceInsightsCore/Private/InsightsCore/Widgets/SSegmentedBarGraph.cpp
--- a/Engine/Source/Developer/TraceInsightsCore/Private/InsightsCore/Widgets/SSegmentedBarGraph.cpp
+++ b/Engine/Source/Developer/TraceInsightsCore/Private/InsightsCore/Widgets/SSegmentedBarGraph.cpp
@@ -4,8 +4,10 @@
 
 #include "SlateOptMacros.h"
 #include "Styling/AppStyle.h"
+#include "Styling/StyleColors.h"
 #include "Widgets/Images/SImage.h"
 #include "Widgets/Layout/SBorder.h"
+#include "Widgets/Layout/SBox.h"
 #include "Widgets/Text/STextBlock.h"
 
 #define LOCTEXT_NAMESPACE "UE::Insights::SSegmentedBarGraph"
@@ -15,11 +17,21 @@ namespace UE::Insights
 
 ////////////////////////////////////////////////////////////////////////////////////////////////////
 
-void FSharedObservableSegmentArray::HandleArrayChanged(typename ::UE::Slate::Containers::TObservableArray<IBarGraphSegment>::ObservableArrayChangedArgsType Args)
+void FObservableSegmentArrayPointer::HandleArrayChanged(typename ::UE::Slate::Containers::TObservableArray<FObservableSegmentArrayPointer::SegmentType>::ObservableArrayChangedArgsType Args)
 {
-	if (TSharedPtr<SSegmentedBarGraph> GraphOwnerPin = GraphOwner.Pin())
+	if (TSharedPtr<WidgetType> GraphWidgetOwnerPin = GraphWidgetOwner.Pin())
 	{
-		GraphOwnerPin->RequestGraphRefresh();
+		GraphWidgetOwnerPin->RequestGraphRefresh();
+	}
+}
+
+////////////////////////////////////////////////////////////////////////////////////////////////////
+
+void FSharedObservableSegmentArray::HandleArrayChanged(typename ::UE::Slate::Containers::TObservableArray<FSharedObservableSegmentArray::SegmentType>::ObservableArrayChangedArgsType Args)
+{
+	if (TSharedPtr<SSegmentedBarGraph> GraphWidgetOwnerPin = GraphWidgetOwner.Pin())
+	{
+		GraphWidgetOwnerPin->RequestGraphRefresh();
 	}
 }
 
@@ -42,9 +54,18 @@ SSegmentedBarGraph::~SSegmentedBarGraph()
 BEGIN_SLATE_FUNCTION_BUILD_OPTIMIZATION
 void SSegmentedBarGraph::Construct(const FArguments& InArgs)
 {
-	SAssignNew(HorizontalBox, SHorizontalBox);
-	check(HorizontalBox.IsValid());
-	ChildSlot[ HorizontalBox.ToSharedRef() ];
+	this->SetSegmentsSource(InArgs.MakeSegmentsSource(this->SharedThis(this)));
+
+	ChildSlot
+		[
+			SNew(SBorder)
+			.Padding(0.0f)
+			.BorderImage(FAppStyle::Get().GetBrush("WhiteBrush"))
+			.BorderBackgroundColor(FSlateColor(EStyleColor::Background))
+			[
+				SAssignNew(HorizontalBox, SHorizontalBox)
+			]
+		];
 
 	RequestGraphRefresh();
 }
@@ -54,13 +75,15 @@ END_SLATE_FUNCTION_BUILD_OPTIMIZATION
 
 void SSegmentedBarGraph::RequestGraphRefresh()
 {
-	check(HorizontalBox.IsValid())
-	HorizontalBox->ClearChildren();
-
-	TArrayView<const IBarGraphSegment> Segments = GetSegments();
-	for (int32 SegmentIndex = 0; SegmentIndex < Segments.Num(); ++SegmentIndex)
+	if (HorizontalBox.IsValid())
 	{
-		AddHorizontalSegment(HorizontalBox.ToSharedRef(), SegmentIndex);
+		HorizontalBox->ClearChildren();
+
+		TArrayView<const TSharedPtr<IBarGraphSegment>> Segments = GetSegments();
+		for (int32 SegmentIndex = 0; SegmentIndex < Segments.Num(); ++SegmentIndex)
+		{
+			AddHorizontalSegment(HorizontalBox.ToSharedRef(), SegmentIndex);
+		}
 	}
 }
 
@@ -72,18 +95,23 @@ void SSegmentedBarGraph::AddHorizontalSegment(TSharedRef<SHorizontalBox> Box, in
 	Box->AddSlot()
 		.HAlign(HAlign_Fill)
 		.VAlign(VAlign_Fill)
-		.Padding(FMargin(0.0f, 0.0f))
+		.Padding(FMargin(1.0f, 1.0f))
 		.FillWidth(TAttribute<float>::CreateSP(this, &SSegmentedBarGraph::GetSegmentSize, SegmentIndex))
 		[
 			SNew(SBorder)
+			.Padding(0.0f)
 			.BorderImage(FAppStyle::Get().GetBrush("WhiteBrush"))
-			.BorderBackgroundColor(TAttribute<FSlateColor>::CreateSP(this, &SSegmentedBarGraph::GetSegmentColor, SegmentIndex))
-			.ColorAndOpacity(TAttribute<FLinearColor>::CreateSP(this, &SSegmentedBarGraph::GetSegmentBorderColor, SegmentIndex))
+			.BorderBackgroundColor(this, &SSegmentedBarGraph::GetSegmentColor, SegmentIndex)
 			.ToolTipText(this, &SSegmentedBarGraph::GetSegmentToolTipText, SegmentIndex)
 			[
-				SNew(STextBlock)
-				.ColorAndOpacity(TAttribute<FSlateColor>::CreateSP(this, &SSegmentedBarGraph::GetSegmentTextColor, SegmentIndex))
-				.Text(this, &SSegmentedBarGraph::GetSegmentText, SegmentIndex)
+				SNew(SBox)
+				.VAlign(VAlign_Center)
+				[
+					SNew(STextBlock)
+					.Margin(FMargin(2.0f, 0.0f))
+					.ColorAndOpacity(this, &SSegmentedBarGraph::GetSegmentTextColor, SegmentIndex)
+					.Text(this, &SSegmentedBarGraph::GetSegmentText, SegmentIndex)
+				]
 			]
 		];
 }
@@ -93,77 +121,62 @@ END_SLATE_FUNCTION_BUILD_OPTIMIZATION
 
 float SSegmentedBarGraph::GetSegmentSize(int32 SegmentIndex) const
 {
-	TArrayView<const IBarGraphSegment> Segments = GetSegments();
+	TArrayView<const TSharedPtr<IBarGraphSegment>> Segments = GetSegments();
 	if (SegmentIndex < 0 || SegmentIndex >= Segments.Num())
 	{
 		return 0;
 	}
-	return static_cast<float>(Segments[SegmentIndex].GetSize());
+	return static_cast<float>(Segments[SegmentIndex]->GetSize());
 }
 
 ////////////////////////////////////////////////////////////////////////////////////////////////////
 
 FText SSegmentedBarGraph::GetSegmentText(int32 SegmentIndex) const
 {
-	TArrayView<const IBarGraphSegment> Segments = GetSegments();
+	TArrayView<const TSharedPtr<IBarGraphSegment>> Segments = GetSegments();
 	if (SegmentIndex < 0 || SegmentIndex >= Segments.Num())
 	{
 		return FText::GetEmpty();
 	}
-	return Segments[SegmentIndex].GetText();
+	return Segments[SegmentIndex]->GetText();
 }
 
 ////////////////////////////////////////////////////////////////////////////////////////////////////
 
 FText SSegmentedBarGraph::GetSegmentToolTipText(int32 SegmentIndex) const
 {
-	TArrayView<const IBarGraphSegment> Segments = GetSegments();
+	TArrayView<const TSharedPtr<IBarGraphSegment>> Segments = GetSegments();
 	if (SegmentIndex < 0 || SegmentIndex >= Segments.Num())
 	{
 		return FText::GetEmpty();
 	}
-	return Segments[SegmentIndex].GetToolTipText();
+	return Segments[SegmentIndex]->GetToolTipText();
 }
 
 ////////////////////////////////////////////////////////////////////////////////////////////////////
 
 FSlateColor SSegmentedBarGraph::GetSegmentColor(int32 SegmentIndex) const
 {
-	TArrayView<const IBarGraphSegment> Segments = GetSegments();
+	TArrayView<const TSharedPtr<IBarGraphSegment>> Segments = GetSegments();
 	if (SegmentIndex < 0 || SegmentIndex >= Segments.Num())
 	{
 		FLinearColor Color(0.75f, 0.75f, 0.75f, 1.0f);
 		return FSlateColor(Color);
 	}
-	return FSlateColor(Segments[SegmentIndex].GetColor());
+	return FSlateColor(Segments[SegmentIndex]->GetColor());
 }
 
 ////////////////////////////////////////////////////////////////////////////////////////////////////
 
 FSlateColor SSegmentedBarGraph::GetSegmentTextColor(int32 SegmentIndex) const
 {
-	TArrayView<const IBarGraphSegment> Segments = GetSegments();
+	TArrayView<const TSharedPtr<IBarGraphSegment>> Segments = GetSegments();
 	if (SegmentIndex < 0 || SegmentIndex >= Segments.Num())
 	{
 		FLinearColor Color(0.0f, 0.0f, 0.0f, 1.0f);
 		return FSlateColor(Color);
 	}
-	return FSlateColor(Segments[SegmentIndex].GetTextColor());
-}
-
-////////////////////////////////////////////////////////////////////////////////////////////////////
-
-FLinearColor SSegmentedBarGraph::GetSegmentBorderColor(int32 SegmentIndex) const
-{
-	TArrayView<const IBarGraphSegment> Segments = GetSegments();
-	if (SegmentIndex < 0 || SegmentIndex >= Segments.Num())
-	{
-		return FLinearColor(0.56f, 0.56f, 0.56f, 1.0f);
-	}
-	FLinearColor Color = Segments[SegmentIndex].GetColor();
-	Color *= 0.75f;
-	Color.A = 1.0f;
-	return Color;
+	return FSlateColor(Segments[SegmentIndex]->GetTextColor());
 }
 
 ////////////////////////////////////////////////////////////////////////////////////////////////////
diff --git a/Engine/Source/Developer/TraceInsightsCore/Public/InsightsCore/ViewModels/IBarGraphSegment.h b/Engine/Source/Developer/TraceInsightsCore/Public/InsightsCore/ViewModels/IBarGraphSegment.h
--- a/Engine/Source/Developer/TraceInsightsCore/Public/InsightsCore/ViewModels/IBarGraphSegment.h
+++ b/Engine/Source/Developer/TraceInsightsCore/Public/InsightsCore/ViewModels/IBarGraphSegment.h
@@ -11,7 +11,7 @@
 namespace UE::Insights
 {
 
-class IBarGraphSegment : public TSharedFromThis<IBarGraphSegment>
+class IBarGraphSegment
 {
 public:
 	virtual ~IBarGraphSegment() = default;
diff --git a/Engine/Source/Developer/TraceInsightsCore/Public/InsightsCore/Widgets/SSegmentedBarGraph.h b/Engine/Source/Developer/TraceInsightsCore/Public/InsightsCore/Widgets/SSegmentedBarGraph.h
--- a/Engine/Source/Developer/TraceInsightsCore/Public/InsightsCore/Widgets/SSegmentedBarGraph.h
+++ b/Engine/Source/Developer/TraceInsightsCore/Public/InsightsCore/Widgets/SSegmentedBarGraph.h
@@ -18,42 +18,86 @@
 namespace UE::Insights
 {
 
+class SSegmentedBarGraph;
+
 ////////////////////////////////////////////////////////////////////////////////////////////////////
 
-class SSegmentedBarGraph;
+class FObservableSegmentArrayPointer : public UE::Slate::ItemsSource::IItemsSource<TSharedPtr<IBarGraphSegment>>
+{
+public:
+	using WidgetType = SSegmentedBarGraph;
+	using SegmentType = TSharedPtr<IBarGraphSegment>;
 
-class FSharedObservableSegmentArray : public UE::Slate::ItemsSource::IItemsSource<IBarGraphSegment>
+	explicit FObservableSegmentArrayPointer(TSharedRef<WidgetType> InGraphWidget, UE::Slate::Containers::TObservableArray<SegmentType>* InSegmentsSource)
+		: SegmentsSource(InSegmentsSource)
+		, GraphWidgetOwner(InGraphWidget.ToWeakPtr())
+	{
+		ArrayChangedHandle = InSegmentsSource->OnArrayChanged().AddRaw(this, &FObservableSegmentArrayPointer::HandleArrayChanged);
+	}
+
+	virtual ~FObservableSegmentArrayPointer()
+	{
+		// See ~FObservableArrayPointer() in Widgets\Views\IItemsSource.h.
+		checkf(GraphWidgetOwner.IsValid(), TEXT("The widget has a source needed to be released to prevent bad memory access."));
+		SegmentsSource->OnArrayChanged().Remove(ArrayChangedHandle);
+	}
+
+	virtual const TArrayView<const SegmentType> GetItems() const override
+	{
+		return TArrayView<const SegmentType>(SegmentsSource->GetData(), SegmentsSource->Num());
+	}
+
+	virtual bool IsSame(const void* RawPointer) const override
+	{
+		return RawPointer == reinterpret_cast<const void*>(&SegmentsSource);
+	}
+
+private:
+	void HandleArrayChanged(typename UE::Slate::Containers::TObservableArray<SegmentType>::ObservableArrayChangedArgsType Args);
+
+private:
+	UE::Slate::Containers::TObservableArray<SegmentType>* SegmentsSource;
+	TWeakPtr<WidgetType> GraphWidgetOwner;
+	FDelegateHandle ArrayChangedHandle;
+};
+
+////////////////////////////////////////////////////////////////////////////////////////////////////
+
+class FSharedObservableSegmentArray : public UE::Slate::ItemsSource::IItemsSource<TSharedPtr<IBarGraphSegment>>
 {
 public:
-	explicit FSharedObservableSegmentArray(TSharedRef<SSegmentedBarGraph> InGraph, TSharedRef<::UE::Slate::Containers::TObservableArray<IBarGraphSegment>> InItemsSource)
-		: ItemsSource(InItemsSource)
-		, GraphOwner(InGraph.ToWeakPtr())
+	using WidgetType = SSegmentedBarGraph;
+	using SegmentType = TSharedPtr<IBarGraphSegment>;
+
+	explicit FSharedObservableSegmentArray(TSharedRef<WidgetType> InGraphWidget, TSharedRef<UE::Slate::Containers::TObservableArray<SegmentType>> InSegmentsSource)
+		: SegmentsSource(InSegmentsSource)
+		, GraphWidgetOwner(InGraphWidget.ToWeakPtr())
 	{
-		ArrayChangedHandle = InItemsSource->OnArrayChanged().AddRaw(this, &FSharedObservableSegmentArray::HandleArrayChanged);
+		ArrayChangedHandle = InSegmentsSource->OnArrayChanged().AddRaw(this, &FSharedObservableSegmentArray::HandleArrayChanged);
 	}
 
 	virtual ~FSharedObservableSegmentArray()
 	{
-		ItemsSource->OnArrayChanged().Remove(ArrayChangedHandle);
+		SegmentsSource->OnArrayChanged().Remove(ArrayChangedHandle);
 	}
 
-	virtual const TArrayView<const IBarGraphSegment> GetItems() const override
+	virtual const TArrayView<const SegmentType> GetItems() const override
 	{
-		return TArrayView<const IBarGraphSegment>(ItemsSource->GetData(), ItemsSource->Num());
+		return TArrayView<const SegmentType>(SegmentsSource->GetData(), SegmentsSource->Num());
 	}
 
 	virtual bool IsSame(const void* RawPointer) const override
 	{
-		UE::Slate::Containers::TObservableArray<IBarGraphSegment>* ValueToTest = &(ItemsSource.Get());
+		UE::Slate::Containers::TObservableArray<SegmentType>* ValueToTest = &(SegmentsSource.Get());
 		return RawPointer == reinterpret_cast<const void*>(ValueToTest);
 	}
 
 private:
-	void HandleArrayChanged(typename ::UE::Slate::Containers::TObservableArray<IBarGraphSegment>::ObservableArrayChangedArgsType Args);
+	void HandleArrayChanged(typename UE::Slate::Containers::TObservableArray<SegmentType>::ObservableArrayChangedArgsType Args);
 
 private:
-	TSharedRef<UE::Slate::Containers::TObservableArray<IBarGraphSegment>> ItemsSource;
-	TWeakPtr<SSegmentedBarGraph> GraphOwner;
+	TSharedRef<UE::Slate::Containers::TObservableArray<SegmentType>> SegmentsSource;
+	TWeakPtr<WidgetType> GraphWidgetOwner;
 	FDelegateHandle ArrayChangedHandle;
 };
 
@@ -65,6 +109,9 @@ private:
 class SSegmentedBarGraph : public SCompoundWidget
 {
 public:
+	using SegmentType = TSharedPtr<IBarGraphSegment>;
+	using SegmentSourceType = UE::Slate::ItemsSource::IItemsSource<SegmentType>;
+
 	/** Default constructor. */
 	UE_API SSegmentedBarGraph();
 
@@ -72,7 +119,23 @@ public:
 	UE_API virtual ~SSegmentedBarGraph();
 
 	SLATE_BEGIN_ARGS(SSegmentedBarGraph) {}
-		SLATE_ITEMS_SOURCE_ARGUMENT(IBarGraphSegment, SegmentsSource)
+		SLATE_ITEMS_SOURCE_ARGUMENT(SegmentType, SegmentsSource)
+		TUniquePtr<SegmentSourceType> MakeSegmentsSource(TSharedRef<SSegmentedBarGraph> InWidget) const
+		{
+			if (_SegmentsSource_ArrayPointer)
+			{
+				return MakeUnique<UE::Slate::ItemsSource::FArrayPointer<SegmentType>>(_SegmentsSource_ArrayPointer);
+			}
+			else if (_SegmentsSource_ObservableArrayPointer)
+			{
+				return MakeUnique<FObservableSegmentArrayPointer>(InWidget, _SegmentsSource_ObservableArrayPointer);
+			}
+			else if (_SegmentsSource_SharedObservableArray)
+			{
+				return MakeUnique<FSharedObservableSegmentArray>(InWidget, _SegmentsSource_SharedObservableArray.ToSharedRef());
+			}
+			return TUniquePtr<SegmentSourceType>();
+		}
 	SLATE_END_ARGS()
 
 	/**
@@ -87,14 +150,14 @@ public:
 	 * Set the SegmentsSource. The graph bar will generate widgets to represent these segments.
 	 * @param InSegmentsSource A pointer to the array of segments that should be observed by this widget.
 	*/
-	void SetSegmentsSource(const TArray<IBarGraphSegment>* InSegmentsSource)
+	void SetSegmentsSource(const TArray<SegmentType>* InSegmentsSource)
 	{
 		ensureMsgf(InSegmentsSource, TEXT("The SegmentsSource is invalid."));
 		if (SegmentsSource == nullptr || !SegmentsSource->IsSame(reinterpret_cast<const void*>(InSegmentsSource)))
 		{
 			if (InSegmentsSource)
 			{
-				SetSegmentsSource(MakeUnique<UE::Slate::ItemsSource::FArrayPointer<IBarGraphSegment>>(InSegmentsSource));
+				SetSegmentsSource(MakeUnique<UE::Slate::ItemsSource::FArrayPointer<SegmentType>>(InSegmentsSource));
 			}
 			else
 			{
@@ -103,7 +166,7 @@ public:
 		}
 	}
 
-	void SetSegmentsSource(TSharedRef<UE::Slate::Containers::TObservableArray<IBarGraphSegment>> InSegmentsSource)
+	void SetSegmentsSource(TSharedRef<UE::Slate::Containers::TObservableArray<SegmentType>> InSegmentsSource)
 	{
 		if (SegmentsSource == nullptr || !SegmentsSource->IsSame(reinterpret_cast<const void*>(&InSegmentsSource.Get())))
 		{
@@ -111,7 +174,7 @@ public:
 		}
 	}
 
-	void SetSegmentsSource(TUniquePtr<UE::Slate::ItemsSource::IItemsSource<IBarGraphSegment>> Provider)
+	void SetSegmentsSource(TUniquePtr<SegmentSourceType> Provider)
 	{
 		SegmentsSource = MoveTemp(Provider);
 		RequestGraphRefresh();
@@ -119,7 +182,7 @@ public:
 
 	void ClearSegmentsSource()
 	{
-		SetSegmentsSource(TUniquePtr<UE::Slate::ItemsSource::IItemsSource<IBarGraphSegment>>());
+		SetSegmentsSource(TUniquePtr<SegmentSourceType>());
 	}
 
 	bool HasValidSegmentsSource() const
@@ -127,9 +190,9 @@ public:
 		return SegmentsSource != nullptr;
 	}
 
-	TArrayView<const IBarGraphSegment> GetSegments() const
+	TArrayView<const SegmentType> GetSegments() const
 	{
-		return SegmentsSource ? SegmentsSource->GetItems() : TArrayView<const IBarGraphSegment>();
+		return SegmentsSource ? SegmentsSource->GetItems() : TArrayView<const SegmentType>();
 	}
 
 private:
@@ -139,11 +202,11 @@ private:
 	FText GetSegmentToolTipText(int32 SegmentIndex) const;
 	FSlateColor GetSegmentColor(int32 SegmentIndex) const;
 	FSlateColor GetSegmentTextColor(int32 SegmentIndex) const;
-	FLinearColor GetSegmentBorderColor(int32 SegmentIndex) const;
 
 private:
 	/** Pointer to the source data that we are observing */
-	TUniquePtr<UE::Slate::ItemsSource::IItemsSource<IBarGraphSegment>> SegmentsSource;
+	TUniquePtr<SegmentSourceType> SegmentsSource;
+
 	TSharedPtr<SHorizontalBox> HorizontalBox;
 };
 
