[TraceInsightsCore] Added FTableCellValueSorter::SetShouldSortGroupsFirst(bool) to allow sorting to mix group and leaf nodes.

#Synced-CL 52393202

diff --git a/Engine/Source/Developer/TraceInsightsCore/Private/InsightsCore/Table/ViewModels/TableCellValueSorter.cpp b/Engine/Source/Developer/TraceInsightsCore/Private/InsightsCore/Table/ViewModels/TableCellValueSorter.cpp
--- a/Engine/Source/Developer/TraceInsightsCore/Private/InsightsCore/Table/ViewModels/TableCellValueSorter.cpp
+++ b/Engine/Source/Developer/TraceInsightsCore/Private/InsightsCore/Table/ViewModels/TableCellValueSorter.cpp
@@ -13,7 +13,7 @@
 		{ \
 			return CancelSort(); \
 		} \
-		if (A->IsGroup() != B->IsGroup()) \
+		if (ShouldSortGroupsFirst() && A->IsGroup() != B->IsGroup()) \
 		{ \
 			return A->IsGroup(); \
 		} \
@@ -202,7 +202,7 @@ FSorterByBoolValue::FSorterByBoolValue(TSharedRef<FTableColumn> InColumnRef)
 		}
 	};
 
-	DescendingCompareDelegate = [&Column = * InColumnRef, this](const FBaseTreeNodePtr& A, const FBaseTreeNodePtr& B) -> bool
+	DescendingCompareDelegate = [&Column = *InColumnRef, this](const FBaseTreeNodePtr& A, const FBaseTreeNodePtr& B) -> bool
 	{
 		INSIGHTS_DEFAULT_PRESORTING_NODES(A, B)
 
diff --git a/Engine/Source/Developer/TraceInsightsCore/Public/InsightsCore/Table/ViewModels/TableCellValueSorter.h b/Engine/Source/Developer/TraceInsightsCore/Public/InsightsCore/Table/ViewModels/TableCellValueSorter.h
--- a/Engine/Source/Developer/TraceInsightsCore/Public/InsightsCore/Table/ViewModels/TableCellValueSorter.h
+++ b/Engine/Source/Developer/TraceInsightsCore/Public/InsightsCore/Table/ViewModels/TableCellValueSorter.h
@@ -69,6 +69,9 @@ public:
 
 	virtual void SetAsyncOperationProgress(IAsyncOperationProgress* InAsyncOperationProgress) override { AsyncOperationProgress = InAsyncOperationProgress; }
 
+	bool ShouldSortGroupsFirst() const { return bShouldSortGroupsFirst; }
+	void SetShouldSortGroupsFirst(bool bOnOff) { bShouldSortGroupsFirst = bOnOff; }
+
 protected:
 	// Attempts to cancel the sort by throwing an exception. If exceptions are not available the return value is meant to be returned from sort predicates to speed up the sort.
 	UE_API bool CancelSort() const;
@@ -89,6 +92,8 @@ protected:
 	FTreeNodeCompareFunc DescendingCompareDelegate;
 
 	IAsyncOperationProgress* AsyncOperationProgress = nullptr;
+
+	bool bShouldSortGroupsFirst = true;
 };
 
 ////////////////////////////////////////////////////////////////////////////////////////////////////
