[Insights]
* SegmentedbarGraph: Added ellipses as overflow policy for text displayed on segments.
* Fixed invalidation when segments array is updated.

#rnx
#Synced-CL 52677305

diff --git a/Engine/Source/Developer/TraceInsightsCore/Private/InsightsCore/Widgets/SSegmentedBarGraph.cpp b/Engine/Source/Developer/TraceInsightsCore/Private/InsightsCore/Widgets/SSegmentedBarGraph.cpp
--- a/Engine/Source/Developer/TraceInsightsCore/Private/InsightsCore/Widgets/SSegmentedBarGraph.cpp
+++ b/Engine/Source/Developer/TraceInsightsCore/Private/InsightsCore/Widgets/SSegmentedBarGraph.cpp
@@ -84,6 +84,8 @@ void SSegmentedBarGraph::RequestGraphRefresh()
 		{
 			AddHorizontalSegment(HorizontalBox.ToSharedRef(), SegmentIndex);
 		}
+
+		HorizontalBox->Invalidate(EInvalidateWidgetReason::ChildOrder);
 	}
 }
 
@@ -111,6 +113,7 @@ void SSegmentedBarGraph::AddHorizontalSegment(TSharedRef<SHorizontalBox> Box, in
 					.Margin(FMargin(2.0f, 0.0f))
 					.ColorAndOpacity(this, &SSegmentedBarGraph::GetSegmentTextColor, SegmentIndex)
 					.Text(this, &SSegmentedBarGraph::GetSegmentText, SegmentIndex)
+					.OverflowPolicy(ETextOverflowPolicy::Ellipsis)
 				]
 			]
 		];
